Add documentation for async commands in CQTest

#rnx

diff --git a/Engine/Source/Developer/CQTest/README.md b/Engine/Source/Developer/CQTest/README.md
--- a/Engine/Source/Developer/CQTest/README.md
+++ b/Engine/Source/Developer/CQTest/README.md
@@ -242,10 +242,11 @@ Latent actions are supported with the `TEST_CLASS` macro.  Each step will comple
 
 **CQTest** provides the following additional latent actions:
  - `FExecute` - Action that executes only once.
- - `FWaitUntil` - Action that executes over multiple ticks until either completion or the duration exceeds the timeout.  Action will fail if the condition cannot be satisifed before timing out.
+ - `FWaitUntil` - Action that executes over multiple ticks until either completion or the duration exceeds the timeout.  Action will fail if the condition cannot be satisfied before timing out.
  - `FWaitDelay` - Action that waits a specified duration.
    - **CAUTION:** Using a timed-wait can introduce flakiness due to variable runtimes and the above `FWaitUntil` should be used instead.
  - `FRunSequence` - Action which ensures that a collection of latent actions occur in order, and only after all previous actions have finished.
+ - `TAsyncExecute` - Action that executes some asynchronous code, which produces a `TAsyncResult`, and waits for its completion without blocking the Game thread. The action will fail if the result of the asynchronous call is not ready within the specified timeout. The result can be processed in the nested `FExecute` or `FWaitUntil` action.
 
 Also available for commands is a fluent command builder
 ```cpp
@@ -257,11 +258,42 @@ Also available for commands is a fluent command builder
 			.Until([&]() { return StepThreeComplete(); })
 			.Then([&]() { ASSERT_THAT(IsTrue(SomethingImportant)); });
 	}
+	
+	TEST_METHOD(AsyncFuncTest) 
+	{
+		/**
+		 * Assuming the following declarations of async functions:
+		 *
+		 *	TAsyncResult<void> AsyncStepOne();
+		 *	TAsyncResult<void> AsyncStepTwo();
+		 *	TAsyncResult<bool> AsyncStepThree();
+		 *	TAsyncResult<SomeResultType> AsyncStepFour();
+		 */
+
+		TestCommandBuilder
+			.DoAsync<void>([&]() { return AsyncStepOne(); })
+			.ThenAsync<void>([&]() { return AsyncStepTwo(); })
+			.ThenAsync<bool>(
+				// This lambda is called to start an async action
+				[&]() { return AsyncStepThree(); },
+				// This lambda is called to process the result, once it's ready
+				[&](bool bResult) { ASSERT_THAT(IsTrue(bResult)); }
+			)
+			.UntilAsync<SomeResultType>(
+				// This lambda is called to start an async action
+				[&]() { return AsyncStepFour(); },
+				// Once the result is ready, this lambda is called multiple times until the condition is satisfied or the duration exceeds the timeout
+				[&](const SomeResultType& Result) { return Result.StepFourComplete(); }
+			)
+			.Then([&]() { ASSERT_THAT(IsTrue(SomethingImportant)); });
+	}
 ```
 
 The command builder provides commands which wrap around the above mentioned latent actions.  The following commands are made available:
  - `Do`/`Then` - Commands which adds the `FExecute` latent action with the provided lambda to be executed.
  - `StartWhen`/`Until` - Commands which adds the `FWaitUntil` latent action with the provided lambda to be evaluated.
+ - `DoAsync`/`ThenAsync` - Commands which add the `TAsyncExecute` latent action with two provided lambdas: the first one to execute an async function, and the second one (optional) to process the result of the async function in the `FExecute` latent action.
+ - `UntilAsync` - Command which adds the `TAsyncExecute` latent action with two provided lambdas: the first one to execute an async function, and the second one to process the result of the async function in the `FWaitUntil` latent action.
  - `WaitDelay` - Command which waits a specified duration before continuing.  
    - **CAUTION:** Using a timed-wait can introduce flakiness due to variable runtimes and the above `StartWhen`/`Until` commands should be used instead.
  - `OnTearDown`/`CleanUpWith` - Commands which adds the `FExecute` latent action with the provided lambda to be executed after the test.  Can be called multiple times to add multiple clean up latent actions.
