Postpone initialization of TempMapDirectory in the MapTestSpawner component of CQTest.

Accessing FPaths::ProjectContentDir(), which is required for TempMapDirectory initialization, can lead to failures when FCommandLine is not initialized, which can happen when a module inside the Engine folder depends on CQTest.
Therefore, the variable was moved inside the GetTempMapDirectory() function to delay initialization until its first usage.

#rnx

diff --git a/Engine/Source/Developer/CQTest/Private/Components/MapTestSpawner.cpp b/Engine/Source/Developer/CQTest/Private/Components/MapTestSpawner.cpp
--- a/Engine/Source/Developer/CQTest/Private/Components/MapTestSpawner.cpp
+++ b/Engine/Source/Developer/CQTest/Private/Components/MapTestSpawner.cpp
@@ -26,7 +26,11 @@ DEFINE_LOG_CATEGORY_STATIC(LogMapTest, Log, All);
 
 namespace {
 
-static const FString TempMapDirectory = FPaths::Combine(FPaths::ProjectContentDir(), TEXT("CQTestMapTemp"));
+const FString& GetTempMapDirectory()
+{
+	static const FString TempMapDirectory = FPaths::Combine(FPaths::ProjectContentDir(), TEXT("CQTestMapTemp"));
+	return TempMapDirectory;
+}
 
 /**
  * Generates a unique random 8 character map name.
@@ -46,7 +50,7 @@ void CleanupTempResources()
 {
 	const bool bDirectoryMustExist = true;
 	const bool bRemoveRecursively = true;
-	const bool bWasDeleted = IFileManager::Get().DeleteDirectory(*TempMapDirectory, bDirectoryMustExist, bRemoveRecursively);
+	const bool bWasDeleted = IFileManager::Get().DeleteDirectory(*GetTempMapDirectory(), bDirectoryMustExist, bRemoveRecursively);
 	check(bWasDeleted);
 }
 
@@ -74,14 +78,14 @@ TUniquePtr<FMapTestSpawner> FMapTestSpawner::CreateFromTempLevel(FTestCommandBui
 	}
 
 	FString MapName = GenerateUniqueMapName();
-	FString MapPath = FPaths::Combine(TempMapDirectory, MapName);
+	FString MapPath = FPaths::Combine(GetTempMapDirectory(), MapName);
 	FString NewLevelPackage = FPackageName::FilenameToLongPackageName(MapPath);
 
 	ULevelEditorSubsystem* LevelEditorSubsystem = GEditor->GetEditorSubsystem<ULevelEditorSubsystem>();
 	bool bWasTempLevelCreated = LevelEditorSubsystem->NewLevel(NewLevelPackage);
 	check(bWasTempLevelCreated);
 
-	TUniquePtr<FMapTestSpawner> Spawner = MakeUnique<FMapTestSpawner>(TempMapDirectory, MapName);
+	TUniquePtr<FMapTestSpawner> Spawner = MakeUnique<FMapTestSpawner>(GetTempMapDirectory(), MapName);
 	InCommandBuilder.OnTearDown([]() {
 		// Create a new map to free up the reference to the map used during testing before cleaning up all temporary resources
 		FAutomationEditorCommonUtils::CreateNewMap();
