Use accessors for NetDriver::GuidCache which will be made protected in order to trap unexpected usage when running iris replication

Engine files

diff --git a/Engine/Source/Developer/CQTest/Public/Impl/PIENetworkComponent.inl b/Engine/Source/Developer/CQTest/Public/Impl/PIENetworkComponent.inl
--- a/Engine/Source/Developer/CQTest/Public/Impl/PIENetworkComponent.inl
+++ b/Engine/Source/Developer/CQTest/Public/Impl/PIENetworkComponent.inl
@@ -209,7 +209,7 @@ inline FPIENetworkComponent<NetworkDataType>& FPIENetworkComponent<NetworkDataTy
 		}
 	})
 	.UntilServer(TEXT("Waiting for NetGUID"), [this, ServerActor](NetworkDataType& State) {
-		FNetworkGUID NetGUID = State.World->GetNetDriver()->GuidCache->GetNetGUID(*ServerActor);
+		FNetworkGUID NetGUID = State.World->GetNetDriver()->GetNetGuidCache()->GetNetGUID(*ServerActor);
 		if (!NetGUID.IsValid())
 		{
 			return false;
@@ -241,7 +241,7 @@ inline bool FPIENetworkComponent<NetworkDataType>::ReplicateToClients(NetworkDat
 			continue;
 		}
 
-		AActor* ClientActor = Cast<AActor>(ClientState.World->GetNetDriver()->GuidCache->GetObjectFromNetGUID(NetGUID, true));
+		AActor* ClientActor = Cast<AActor>(ClientState.World->GetNetDriver()->GetNetGuidCache()->GetObjectFromNetGUID(NetGUID, true));
 		if (ClientActor != nullptr)
 		{
 			// The other half of the offset implementation.
