Updating FAssetFilterBuilder example to reflect proper usage of returning a FARFilter instead of it being passed in.

#rnx

diff --git a/Engine/Source/Developer/CQTest/Public/Helpers/CQTestAssetHelper.h b/Engine/Source/Developer/CQTest/Public/Helpers/CQTestAssetHelper.h
--- a/Engine/Source/Developer/CQTest/Public/Helpers/CQTestAssetHelper.h
+++ b/Engine/Source/Developer/CQTest/Public/Helpers/CQTestAssetHelper.h
@@ -20,10 +20,9 @@ TEST_CLASS_WITH_FLAGS(MyFixtureName, "AssetHelper.Example", EAutomationTestFlags
 {
 	TEST_METHOD(Object_FromBlueprint_IsLoaded)
 	{
-		FARFilter Filter;
-		CQTestAssetHelper::FAssetFilterBuilder()
+		FARFilter Filter = CQTestAssetHelper::FAssetFilterBuilder()
 			.WithPackagePath(TEXT("/Package/Path/To/Asset"))
-			.Build(Filter);
+			.Build();
 
 		UClass* Class = CQTestAssetHelper::GetBlueprintClass(Filter, TEXT("B_Asset"));
 		ASSERT_THAT(IsNotNull(Class));
