Adding CVars and Developer Settings to allow for per-project default timeouts

PR #12169

diff --git a/Engine/Plugins/Tests/CQTest/Source/CQTestTests/Private/Commands/TestCommandsTests.cpp b/Engine/Plugins/Tests/CQTest/Source/CQTestTests/Private/Commands/TestCommandsTests.cpp
--- a/Engine/Plugins/Tests/CQTest/Source/CQTestTests/Private/Commands/TestCommandsTests.cpp
+++ b/Engine/Plugins/Tests/CQTest/Source/CQTestTests/Private/Commands/TestCommandsTests.cpp
@@ -3,6 +3,8 @@
 #include "CQTest.h"
 #include "CQTestUnitTestHelper.h"
 
+#include "CQTestSettings.h"
+
 TEST_CLASS(RunSequenceBasicTests, "TestFramework.CQTest.Core")
 {
 	TSharedPtr<IAutomationLatentCommand> Cmd1;
@@ -656,4 +658,109 @@ TEST_CLASS(AsyncExecuteTimeoutTests, "TestFramework.CQTest.Core")
 			}
 		));
 	}
-};
\ No newline at end of file
+};
+
+TEST_CLASS(WaitUntilTests, "TestFramework.CQTest.Core")
+{
+	TEST_METHOD(Timeout_WithNoValueProvided_IsGreaterThanZero)
+	{
+		FWaitUntil command(*TestRunner, []() { return true; });
+		ASSERT_THAT(IsTrue(command.Timeout > FTimespan::Zero()));
+	}
+
+	TEST_METHOD(Timeout_WithSpecificValue_IsUsed)
+	{
+		FTimespan Timeout = FTimespan::FromSeconds(100);
+		FWaitUntil command(*TestRunner, []() { return true; }, Timeout);
+		ASSERT_THAT(AreEqual(Timeout, command.Timeout));
+	}
+
+	TEST_METHOD(Timeout_WithDefaultValue_UsesCVar)
+	{
+		FWaitUntil command(*TestRunner, []() { return true; }, CQTest::DefaultTimeout);
+		ASSERT_THAT(IsNear(CQTestConsoleVariables::CommandTimeout, static_cast<float>(command.Timeout.GetTotalSeconds()), SMALL_NUMBER));
+	}
+
+	TEST_METHOD(Timeout_WithOverriddenCvar_UsesOverriddenValue)
+	{
+		IConsoleVariable* ConsoleVariable = IConsoleManager::Get().FindConsoleVariable(CQTestConsoleVariables::CommandTimeoutName);
+		ASSERT_THAT(IsNotNull(ConsoleVariable));
+
+		const float NewTimeout = ConsoleVariable->GetFloat() + 1;
+		FString NewTimeoutStr = FString::Printf(TEXT("%f"), NewTimeout);
+		TSharedPtr<FScopedTestEnvironment> TestEnvironment = UCQTestSettings::SetTestClassTimeouts(FTimespan::FromSeconds(NewTimeout));
+
+		FWaitUntil command(*TestRunner, []() { return true; }, CQTest::DefaultTimeout);
+		ASSERT_THAT(IsNear(NewTimeout, static_cast<float>(command.Timeout.GetTotalSeconds()), SMALL_NUMBER));
+	}
+};
+
+TEST_CLASS(WaitUntilClassTimeoutTest, "TestFramework.CQTest.Core")
+{
+	TEST_METHOD(Timeout_SetInBeforeEach_PersistsInTest)
+	{
+		IConsoleVariable* ConsoleVariable = IConsoleManager::Get().FindConsoleVariable(CQTestConsoleVariables::CommandTimeoutName);
+		ASSERT_THAT(IsNotNull(ConsoleVariable));
+
+		const double TimeoutValue = ConsoleVariable->GetFloat() + 1;
+		TSharedPtr<FScopedTestEnvironment> TestEnvironment = UCQTestSettings::SetTestClassTimeouts(FTimespan::FromSeconds(TimeoutValue));
+
+		FWaitUntil command(*TestRunner, []() { return true; }, CQTest::DefaultTimeout);
+		ASSERT_THAT(IsNear(TimeoutValue, command.Timeout.GetTotalSeconds(), 0.01));
+	}
+};
+
+#if WITH_EDITOR
+// Test checks to make sure that editing the timeout property is reflected. Can only be done within the Editor.
+TEST_CLASS_WITH_FLAGS(WaitUntilUserSettingTimeout, "TestFramework.CQTest.Core", EAutomationTestFlags::EditorContext | EAutomationTestFlags::EngineFilter)
+{
+	double ExpectedTimeout{0.0};
+	FString OriginalTimeout;
+
+	bool SetPropertyValue(const FName& PropertyName, const FString& PropertyValue, FString& OutError)
+	{
+		UCQTestSettings* DefaultSettings = GetMutableDefault<UCQTestSettings>();
+		if (!IsValid(DefaultSettings))
+		{
+			OutError = TEXT("Could not load default CQTest Settings.");
+			return false;
+		}
+
+		if (FProperty* Property = FindFProperty<FFloatProperty>(UCQTestSettings::StaticClass(), PropertyName))
+		{
+			Property->ImportText_InContainer(*PropertyValue, DefaultSettings, DefaultSettings, PPF_None);
+
+			FPropertyChangedEvent ChangeEvent(Property, EPropertyChangeType::ValueSet);
+			DefaultSettings->PostEditChangeProperty(ChangeEvent);
+			return true;
+		}
+
+		OutError = FString::Format(TEXT("Property '{0}' was not found in the CQTest Settings."), { PropertyName.ToString() });
+		return false;
+	}
+
+	BEFORE_EACH()
+	{
+		float DefaultTimeout = GetDefault<UCQTestSettings>()->CommandTimeout;
+		ExpectedTimeout = DefaultTimeout + 1;
+		OriginalTimeout = FString::SanitizeFloat(DefaultTimeout);
+
+		FString ErrorMessage;
+		const bool bWasUpdated = SetPropertyValue(TEXT("CommandTimeout"), FString::SanitizeFloat(ExpectedTimeout), ErrorMessage);
+		AddErrorIfFalse(bWasUpdated, ErrorMessage);
+	}
+
+	AFTER_EACH()
+	{
+		FString ErrorMessage;
+		const bool bWasUpdated = SetPropertyValue(TEXT("CommandTimeout"), OriginalTimeout, ErrorMessage);
+		AddErrorIfFalse(bWasUpdated, ErrorMessage);
+	}
+
+	TEST_METHOD(Timeout_SetInBeforeEach_PersistsInTest)
+	{
+		FWaitUntil command(*TestRunner, []() { return true; }, CQTest::DefaultTimeout);
+		ASSERT_THAT(IsNear(ExpectedTimeout, command.Timeout.GetTotalSeconds(), 0.01));
+	}
+};
+#endif // WITH_EDITOR
\ No newline at end of file
diff --git a/Engine/Source/Developer/CQTest/Private/CQTestSettings.cpp b/Engine/Source/Developer/CQTest/Private/CQTestSettings.cpp
new file mode 100644
--- /dev/null
+++ b/Engine/Source/Developer/CQTest/Private/CQTestSettings.cpp
@@ -0,0 +1,34 @@
+// Copyright Epic Games, Inc. All Rights Reserved.
+
+#include "CQTestSettings.h"
+#include "Tests/AutomationCommon.h"
+
+namespace CQTestConsoleVariables 
+{
+static TAutoConsoleVariable<float> CVarCommandTimeout(
+	CommandTimeoutName,
+	CommandTimeout,
+	TEXT("How long to wait on an asynchronous task before timing out in seconds"));
+
+static TAutoConsoleVariable<float> CVarNetworkTimeout(
+	NetworkTimeoutName,
+	NetworkTimeout,
+	TEXT("How long to wait on a network task before timing out in seconds"));
+
+static TAutoConsoleVariable<float> CVarMapTestTimeout(
+	MapTestTimeoutName,
+	MapTestTimeout,
+	TEXT("How long to wait on a map test before timing out in seconds"));
+} // namespace CQTestConsoleVariables
+
+TSharedPtr<FScopedTestEnvironment> UCQTestSettings::SetTestClassTimeouts(FTimespan Duration)
+{
+	TSharedPtr<FScopedTestEnvironment> TestEnvironment = FScopedTestEnvironment::Get();
+
+	FString DurationString = FString::SanitizeFloat(Duration.GetSeconds());
+	TestEnvironment->SetConsoleVariableValue(CQTestConsoleVariables::CommandTimeoutName, DurationString);
+	TestEnvironment->SetConsoleVariableValue(CQTestConsoleVariables::NetworkTimeoutName, DurationString);
+	TestEnvironment->SetConsoleVariableValue(CQTestConsoleVariables::MapTestTimeoutName, DurationString);
+
+	return TestEnvironment;
+}
\ No newline at end of file
diff --git a/Engine/Source/Developer/CQTest/Private/Commands/TestCommands.cpp b/Engine/Source/Developer/CQTest/Private/Commands/TestCommands.cpp
--- a/Engine/Source/Developer/CQTest/Private/Commands/TestCommands.cpp
+++ b/Engine/Source/Developer/CQTest/Private/Commands/TestCommands.cpp
@@ -1,6 +1,7 @@
 // Copyright Epic Games, Inc. All Rights Reserved.
 
 #include "Commands/TestCommands.h"
+#include "CQTestSettings.h"
 #include "Logging/LogMacros.h"
 
 DEFINE_LOG_CATEGORY_STATIC(LogCqTest, Log, All)
@@ -45,6 +46,23 @@ bool CQTEST_API FWaitUntil::Update()
 	return false;
 }
 
+FTimespan FWaitUntil::MakeTimeout(TOptional<FTimespan> InTimeout)
+{
+	if (InTimeout.IsSet())
+	{
+		return InTimeout.GetValue();
+	}
+	else if (IConsoleVariable* ConsoleVariable = IConsoleManager::Get().FindConsoleVariable(CQTestConsoleVariables::CommandTimeoutName))
+	{
+		return FTimespan::FromSeconds(ConsoleVariable->GetFloat());
+	}
+	else
+	{
+		UE_LOG(LogCqTest, Warning, TEXT("CVar '%s' was not found. Defaulting to %f seconds."), CQTestConsoleVariables::CommandTimeoutName, CQTestConsoleVariables::CommandTimeout);
+		return FTimespan::FromSeconds(CQTestConsoleVariables::CommandTimeout);
+	}
+}
+
 bool CQTEST_API FWaitDelay::Update()
 {
 	if (!bHasTimerStarted)
diff --git a/Engine/Source/Developer/CQTest/Private/Components/MapTestSpawner.cpp b/Engine/Source/Developer/CQTest/Private/Components/MapTestSpawner.cpp
--- a/Engine/Source/Developer/CQTest/Private/Components/MapTestSpawner.cpp
+++ b/Engine/Source/Developer/CQTest/Private/Components/MapTestSpawner.cpp
@@ -4,6 +4,7 @@
 
 #if WITH_AUTOMATION_TESTS
 #include "Commands/TestCommands.h"
+#include "CQTestSettings.h"
 #include "Tests/AutomationCommon.h"
 #include "GameDelegates.h"
 #include "GameFramework/PlayerController.h"
@@ -98,7 +99,7 @@ TUniquePtr<FMapTestSpawner> FMapTestSpawner::CreateFromTempLevel(FTestCommandBui
 #endif // WITH_EDITOR
 }
 
-void FMapTestSpawner::AddWaitUntilLoadedCommand(FAutomationTestBase* TestRunner)
+void FMapTestSpawner::AddWaitUntilLoadedCommand(FAutomationTestBase* TestRunner, TOptional<FTimespan> Timeout)
 {
 	check(PieWorld == nullptr);
 
@@ -115,6 +116,21 @@ void FMapTestSpawner::AddWaitUntilLoadedCommand(FAutomationTestBase* TestRunner)
 	bool bOpened = AutomationOpenMap(LongPackageName, true);
 	check(bOpened);
 
+	FTimespan TimeoutValue;
+	if (Timeout.IsSet())
+	{
+		TimeoutValue = Timeout.GetValue();
+	}
+	else if (IConsoleVariable* ConsoleVariable = IConsoleManager::Get().FindConsoleVariable(CQTestConsoleVariables::MapTestTimeoutName))
+	{
+		TimeoutValue = FTimespan::FromSeconds(ConsoleVariable->GetFloat());
+	}
+	else
+	{
+		UE_LOG(LogMapTest, Warning, TEXT("CVar '%s' was not found. Defaulting to %f seconds."), CQTestConsoleVariables::MapTestTimeoutName, CQTestConsoleVariables::MapTestTimeout);
+		TimeoutValue = FTimespan::FromSeconds(CQTestConsoleVariables::MapTestTimeout);
+	}
+
 	ADD_LATENT_AUTOMATION_COMMAND(FWaitUntil(*TestRunner, [this]() -> bool {
 		for (const FWorldContext& Context : GEngine->GetWorldContexts())
 		{
@@ -137,7 +153,7 @@ void FMapTestSpawner::AddWaitUntilLoadedCommand(FAutomationTestBase* TestRunner)
 		}
 
 		return false;
-	}));
+	}, TimeoutValue));
 }
 
 UWorld* FMapTestSpawner::CreateWorld()
diff --git a/Engine/Source/Developer/CQTest/Private/Components/PIENetworkComponent.cpp b/Engine/Source/Developer/CQTest/Private/Components/PIENetworkComponent.cpp
--- a/Engine/Source/Developer/CQTest/Private/Components/PIENetworkComponent.cpp
+++ b/Engine/Source/Developer/CQTest/Private/Components/PIENetworkComponent.cpp
@@ -19,7 +19,7 @@
 
 DEFINE_LOG_CATEGORY_STATIC(LogNetworkTest, Log, All);
 
-FBasePIENetworkComponent::FBasePIENetworkComponent(FAutomationTestBase* InTestRunner, FTestCommandBuilder& InCommandBuilder, bool IsInitializing)
+FBasePIENetworkComponent::FBasePIENetworkComponent(FAutomationTestBase* InTestRunner, FTestCommandBuilder& InCommandBuilder, bool IsInitializing, TOptional<FTimespan> Timeout)
 	: TestRunner(InTestRunner)
 	, CommandBuilder(&InCommandBuilder)
 {
@@ -28,14 +28,15 @@ FBasePIENetworkComponent::FBasePIENetworkComponent(FAutomationTestBase* InTestRu
 		return;
 	}
 
+	FTimespan TimeoutValue = MakeTimeout(Timeout);
 	CommandBuilder
 		->Do(TEXT("Stop PIE"), [this]() { StopPie(); })
 		.Then(TEXT("Create New Map"), [this]() { FAutomationEditorCommonUtils::CreateNewMap(); })
 		.Then(TEXT("Start PIE"), [this]() { StartPie(); })
-		.Until(TEXT("Set Worlds"), [this]() { return SetWorlds(); })
+		.Until(TEXT("Set Worlds"), [this]() { return SetWorlds(); }, TimeoutValue)
 		.Then(TEXT("Setup Packet Settings"), [this]() { SetPacketSettings(); })
 		.Then(TEXT("Connect Clients to Server"), [this]() { ConnectClientsToServer(); })
-		.Until(TEXT("Await Clients Ready"), [this]() { return AwaitClientsReady(); })
+		.Until(TEXT("Await Clients Ready"), [this]() { return AwaitClientsReady(); }, TimeoutValue)
 		.OnTearDown(TEXT("Restore Editor State"), [this]() { RestoreState(); });
 }
 
@@ -62,27 +63,31 @@ FBasePIENetworkComponent& FBasePIENetworkComponent::Do(const TCHAR* Description,
 	return *this;
 }
 
-FBasePIENetworkComponent& FBasePIENetworkComponent::Until(TFunction<bool()> Query, FTimespan Timeout)
+FBasePIENetworkComponent& FBasePIENetworkComponent::Until(TFunction<bool()> Query, TOptional<FTimespan> Timeout)
 {
-	CommandBuilder->Until(Query, Timeout);
+	FTimespan TimeoutValue = MakeTimeout(Timeout);
+	CommandBuilder->Until(Query, TimeoutValue);
 	return *this;
 }
 
-FBasePIENetworkComponent& FBasePIENetworkComponent::Until(const TCHAR* Description, TFunction<bool()> Query, FTimespan Timeout)
+FBasePIENetworkComponent& FBasePIENetworkComponent::Until(const TCHAR* Description, TFunction<bool()> Query, TOptional<FTimespan> Timeout)
 {
-	CommandBuilder->Until(Description, Query, Timeout);
+	FTimespan TimeoutValue = MakeTimeout(Timeout);
+	CommandBuilder->Until(Description, Query, TimeoutValue);
 	return *this;
 }
 
-FBasePIENetworkComponent& FBasePIENetworkComponent::StartWhen(TFunction<bool()> Query, FTimespan Timeout)
+FBasePIENetworkComponent& FBasePIENetworkComponent::StartWhen(TFunction<bool()> Query, TOptional<FTimespan> Timeout)
 {
-	CommandBuilder->StartWhen(Query, Timeout);
+	FTimespan TimeoutValue = MakeTimeout(Timeout);
+	CommandBuilder->StartWhen(Query, TimeoutValue);
 	return *this;
 }
 
-FBasePIENetworkComponent& FBasePIENetworkComponent::StartWhen(const TCHAR* Description, TFunction<bool()> Query, FTimespan Timeout)
+FBasePIENetworkComponent& FBasePIENetworkComponent::StartWhen(const TCHAR* Description, TFunction<bool()> Query, TOptional<FTimespan> Timeout)
 {
-	CommandBuilder->StartWhen(Description, Query, Timeout);
+	FTimespan TimeoutValue = MakeTimeout(Timeout);
+	CommandBuilder->StartWhen(Description, Query, TimeoutValue);
 	return *this;
 }
 
@@ -269,4 +274,22 @@ void FBasePIENetworkComponent::RestoreState()
 		StateRestorer.Restore();
 	}
 }
+
+FTimespan FBasePIENetworkComponent::MakeTimeout(TOptional<FTimespan> Timeout)
+{
+	if (Timeout.IsSet())
+	{
+		return Timeout.GetValue();
+	}
+	else if (IConsoleVariable* ConsoleVariable = IConsoleManager::Get().FindConsoleVariable(CQTestConsoleVariables::NetworkTimeoutName))
+	{
+		return FTimespan::FromSeconds(ConsoleVariable->GetFloat());
+	}
+	else
+	{
+		UE_LOG(LogNetworkTest, Warning, TEXT("CVar '%s' was not found. Defaulting to %f seconds."), CQTestConsoleVariables::NetworkTimeoutName, CQTestConsoleVariables::NetworkTimeout);
+		return FTimespan::FromSeconds(CQTestConsoleVariables::NetworkTimeout);
+	}
+}
+
 #endif // ENABLE_PIE_NETWORK_TEST
\ No newline at end of file
diff --git a/Engine/Source/Developer/CQTest/Public/CQTestSettings.h b/Engine/Source/Developer/CQTest/Public/CQTestSettings.h
new file mode 100644
--- /dev/null
+++ b/Engine/Source/Developer/CQTest/Public/CQTestSettings.h
@@ -0,0 +1,65 @@
+// Copyright Epic Games, Inc. All Rights Reserved.
+
+#pragma once
+
+#include "Engine/DeveloperSettingsBackedByCVars.h"
+#include "HAL/IConsoleManager.h"
+#include "Misc/ConfigUtilities.h"
+
+#include "CQTestSettings.generated.h"
+
+struct FScopedTestEnvironment;
+
+namespace CQTestConsoleVariables 
+{
+inline static constexpr float CommandTimeout = 10.0f;
+constexpr auto CommandTimeoutName = TEXT("TestFramework.CQTest.CommandTimeout");
+
+inline static constexpr float NetworkTimeout = 30.0f;
+constexpr auto NetworkTimeoutName = TEXT("TestFramework.CQTest.CommandTimeout.Network");
+
+inline static constexpr float MapTestTimeout = 30.0f;
+constexpr auto MapTestTimeoutName = TEXT("TestFramework.CQTest.CommandTimeout.MapTest");
+} // namespace CQTestConsoleVariables
+
+/** Implements per project Engine settings for the CQTest plugin. */
+UCLASS(config = Engine, defaultconfig, meta = (DisplayName = "CQ Test Settings"))
+class CQTEST_API UCQTestSettings : public UDeveloperSettingsBackedByCVars
+{
+	GENERATED_BODY()
+
+public:
+	void PostInitProperties() override
+	{
+		if (IsTemplate())
+		{
+			// We want the .ini file to have precedence over the CVar constructor, so we apply the ini to the CVar before following the regular UDeveloperSettingsBackedByCVars flow
+			UE::ConfigUtilities::ApplyCVarSettingsFromIni(TEXT("/Script/CQTest.CQTestSettings"), *GEngineIni, ECVF_SetByProjectSetting);
+		}
+
+		Super::PostInitProperties();
+	}
+
+	/** Timeout for WaitUntil latent actions. */
+	UPROPERTY(config, EditAnywhere, Category = "Test Settings", meta = (ConsoleVariable = "TestFramework.CQTest.CommandTimeout"))
+	float CommandTimeout = CQTestConsoleVariables::CommandTimeout;
+
+	/** Timeout for WaitUntil latent actions from the PIENetworkComponent. */
+	UPROPERTY(config, EditAnywhere, Category = "Test Settings", meta = (ConsoleVariable = "TestFramework.CQTest.CommandTimeout.Network"))
+	float NetworkTimeout = CQTestConsoleVariables::NetworkTimeout;
+
+	/** Timeout for `FMapTestSpawner::AddWaitUntilLoadedCommand` latent action used during UWorld loading. */
+	UPROPERTY(config, EditAnywhere, Category = "Test Settings", meta = (ConsoleVariable = "TestFramework.CQTest.CommandTimeout.MapTest"))
+	float MapTestTimeout = CQTestConsoleVariables::MapTestTimeout;
+
+	/**
+	 * Sets the duration for all available timeouts
+	 *
+	 * @param Duration - Duration of the timeout.
+	 * 
+	 * @return shared pointer to the FScopedTestEnvironment singleton instance
+	 *
+	 * @note Console Variable will reset back to the value prior to getting set once the FScopedTestEnvironment resets or gets destructed.
+	 */
+	[[nodiscard]] static TSharedPtr<FScopedTestEnvironment> SetTestClassTimeouts(FTimespan Duration);
+};
\ No newline at end of file
diff --git a/Engine/Source/Developer/CQTest/Public/Commands/TestCommandBuilder.h b/Engine/Source/Developer/CQTest/Public/Commands/TestCommandBuilder.h
--- a/Engine/Source/Developer/CQTest/Public/Commands/TestCommandBuilder.h
+++ b/Engine/Source/Developer/CQTest/Public/Commands/TestCommandBuilder.h
@@ -41,7 +41,7 @@ public:
 		return Do(Description, Action);
 	}
 
-	FTestCommandBuilder& Until(const TCHAR* Description, TFunction<bool()> Query, FTimespan Timeout = FTimespan::FromSeconds(10))
+	FTestCommandBuilder& Until(const TCHAR* Description, TFunction<bool()> Query, TOptional<FTimespan> Timeout = CQTest::DefaultTimeout)
 	{
 		if (!TestRunner.HasAnyErrors())
 		{
@@ -50,13 +50,13 @@ public:
 		return *this;
 	}
 
-	FTestCommandBuilder& Until(TFunction<bool()> Query, FTimespan Timeout = FTimespan::FromSeconds(10))
+	FTestCommandBuilder& Until(TFunction<bool()> Query, TOptional<FTimespan> Timeout = CQTest::DefaultTimeout)
 	{
 		return Until(nullptr, Query, Timeout);
 	}
 
 	template <typename ResultType>
-	FTestCommandBuilder& DoAsync(const TCHAR* Description, TFunction<TAsyncResult<ResultType>()> AsyncAction, FTimespan Timeout = FTimespan::FromSeconds(10))
+	FTestCommandBuilder& DoAsync(const TCHAR* Description, TFunction<TAsyncResult<ResultType>()> AsyncAction, TOptional<FTimespan> Timeout = CQTest::DefaultTimeout)
 	{
 		if (!TestRunner.HasAnyErrors())
 		{
@@ -66,7 +66,7 @@ public:
 	}
 
 	template <typename ResultType>
-	FTestCommandBuilder& DoAsync(TFunction<TAsyncResult<ResultType>()> AsyncAction, FTimespan Timeout = FTimespan::FromSeconds(10))
+	FTestCommandBuilder& DoAsync(TFunction<TAsyncResult<ResultType>()> AsyncAction, TOptional<FTimespan> Timeout = CQTest::DefaultTimeout)
 	{
 		return DoAsync<ResultType>(nullptr, MoveTemp(AsyncAction), Timeout);
 	}
@@ -76,7 +76,7 @@ public:
 		const TCHAR* Description,
 		TFunction<TAsyncResult<ResultType>()> AsyncAction,
 		TFunction<void(typename TAsyncResultCallbackArg<ResultType>::Type)> ResultCallback,
-		FTimespan Timeout = FTimespan::FromSeconds(10))
+		TOptional<FTimespan> Timeout = CQTest::DefaultTimeout)
 	{
 		if (!TestRunner.HasAnyErrors())
 		{
@@ -89,19 +89,19 @@ public:
 	typename TEnableIf<!std::is_void_v<ResultType>, FTestCommandBuilder&>::Type DoAsync(
 		TFunction<TAsyncResult<ResultType>()> AsyncAction,
 		TFunction<void(typename TAsyncResultCallbackArg<ResultType>::Type)> ResultCallback,
-		FTimespan Timeout = FTimespan::FromSeconds(10))
+		TOptional<FTimespan> Timeout = CQTest::DefaultTimeout)
 	{
 		return DoAsync<ResultType>(nullptr, MoveTemp(AsyncAction), MoveTemp(ResultCallback), Timeout);
 	}
 
 	template <typename ResultType>
-	FTestCommandBuilder& ThenAsync(const TCHAR* Description, TFunction<TAsyncResult<ResultType>()> AsyncAction, FTimespan Timeout = FTimespan::FromSeconds(10))
+	FTestCommandBuilder& ThenAsync(const TCHAR* Description, TFunction<TAsyncResult<ResultType>()> AsyncAction, TOptional<FTimespan> Timeout = CQTest::DefaultTimeout)
 	{
 		return DoAsync<ResultType>(Description, MoveTemp(AsyncAction), Timeout);
 	}
 
 	template <typename ResultType>
-	FTestCommandBuilder& ThenAsync(TFunction<TAsyncResult<ResultType>()> AsyncAction, FTimespan Timeout = FTimespan::FromSeconds(10))
+	FTestCommandBuilder& ThenAsync(TFunction<TAsyncResult<ResultType>()> AsyncAction, TOptional<FTimespan> Timeout = CQTest::DefaultTimeout)
 	{
 		return DoAsync<ResultType>(nullptr, MoveTemp(AsyncAction), Timeout);
 	}
@@ -111,7 +111,7 @@ public:
 		const TCHAR* Description,
 		TFunction<TAsyncResult<ResultType>()> AsyncAction,
 		TFunction<void(typename TAsyncResultCallbackArg<ResultType>::Type)> ResultCallback,
-		FTimespan Timeout = FTimespan::FromSeconds(10))
+		TOptional<FTimespan> Timeout = CQTest::DefaultTimeout)
 	{
 		return DoAsync<ResultType>(Description, MoveTemp(AsyncAction), MoveTemp(ResultCallback), Timeout);
 	}
@@ -120,7 +120,7 @@ public:
 	typename TEnableIf<!std::is_void_v<ResultType>, FTestCommandBuilder&>::Type ThenAsync(
 		TFunction<TAsyncResult<ResultType>()> AsyncAction,
 		TFunction<void(typename TAsyncResultCallbackArg<ResultType>::Type)> ResultCallback,
-		FTimespan Timeout = FTimespan::FromSeconds(10))
+		TOptional<FTimespan> Timeout = CQTest::DefaultTimeout)
 	{
 		return DoAsync<ResultType>(nullptr, MoveTemp(AsyncAction), MoveTemp(ResultCallback), Timeout);
 	}
@@ -130,8 +130,8 @@ public:
 		const TCHAR* Description,
 		TFunction<TAsyncResult<ResultType>()> AsyncAction,
 		TFunction<bool(typename TAsyncResultCallbackArg<ResultType>::Type)> ResultCallback,
-		FTimespan AsyncActionTimeout = FTimespan::FromSeconds(10),
-		FTimespan ConditionTimeout = FTimespan::FromSeconds(10))
+		TOptional<FTimespan> AsyncActionTimeout = CQTest::DefaultTimeout,
+		TOptional<FTimespan> ConditionTimeout = CQTest::DefaultTimeout)
 	{
 		if (!TestRunner.HasAnyErrors())
 		{
@@ -144,18 +144,18 @@ public:
 	typename TEnableIf<!std::is_void_v<ResultType>, FTestCommandBuilder&>::Type UntilAsync(
 		TFunction<TAsyncResult<ResultType>()> AsyncAction,
 		TFunction<bool(typename TAsyncResultCallbackArg<ResultType>::Type)> ResultCallback,
-		FTimespan AsyncActionTimeout = FTimespan::FromSeconds(10),
-		FTimespan ConditionTimeout = FTimespan::FromSeconds(10))
+		TOptional<FTimespan> AsyncActionTimeout = CQTest::DefaultTimeout,
+		TOptional<FTimespan> ConditionTimeout = CQTest::DefaultTimeout)
 	{
 		return UntilAsync<ResultType>(nullptr, MoveTemp(AsyncAction), MoveTemp(ResultCallback), AsyncActionTimeout, ConditionTimeout);
 	}
 
-	FTestCommandBuilder& StartWhen(TFunction<bool()> Query, FTimespan Timeout = FTimespan::FromSeconds(10))
+	FTestCommandBuilder& StartWhen(TFunction<bool()> Query, TOptional<FTimespan> Timeout = CQTest::DefaultTimeout)
 	{
 		return Until(Query, Timeout);
 	}
 
-	FTestCommandBuilder& StartWhen(const TCHAR* Description, TFunction<bool()> Query, FTimespan Timeout = FTimespan::FromSeconds(10))
+	FTestCommandBuilder& StartWhen(const TCHAR* Description, TFunction<bool()> Query, TOptional<FTimespan> Timeout = CQTest::DefaultTimeout)
 	{
 		return Until(Description, Query, Timeout);
 	}
diff --git a/Engine/Source/Developer/CQTest/Public/Commands/TestCommands.h b/Engine/Source/Developer/CQTest/Public/Commands/TestCommands.h
--- a/Engine/Source/Developer/CQTest/Public/Commands/TestCommands.h
+++ b/Engine/Source/Developer/CQTest/Public/Commands/TestCommands.h
@@ -5,14 +5,19 @@
 #include "Misc/AutomationTest.h"
 #include "Async/AsyncResult.h"
 
+namespace CQTest
+{
+inline static const TOptional<FTimespan> DefaultTimeout = NullOpt;
+}
+
 /** Latent Command that waits until the Query evaluates to `true` or the timeout has exceeded. */
 class CQTEST_API FWaitUntil : public IAutomationLatentCommand
 {
 public:
-	FWaitUntil(FAutomationTestBase& InTestRunner, TFunction<bool()> Query, FTimespan Timeout = FTimespan::FromSeconds(10), const TCHAR* InDescription = nullptr)
+	FWaitUntil(FAutomationTestBase& InTestRunner, TFunction<bool()> Query, TOptional<FTimespan> Timeout = CQTest::DefaultTimeout, const TCHAR* InDescription = nullptr)
 		: TestRunner(InTestRunner)
-		, Query(MoveTemp(Query))
-		, Timeout(Timeout)
+		, Query(MoveTemp(Query)) 
+		, Timeout(MakeTimeout(Timeout))
 		, Description(InDescription)
 	{}
 
@@ -24,6 +29,9 @@ public:
 	FDateTime StartTime;
 	const TCHAR* Description;
 	bool bHasTimerStarted = false;
+
+private:
+	FTimespan MakeTimeout(TOptional<FTimespan> InTimeout);
 };
 
 /** Latent Command that waits a set time frame.
@@ -113,8 +121,6 @@ class TAsyncExecute;
 /** This namespace provides a set of variables and functions intended for internal use within the TAsyncExecute class */
 namespace AsyncExecuteDetails
 {
-	static const FTimespan DefaultTimeout = FTimespan::FromSeconds(10);
-
 	static const TCHAR* CreateInternalCommandDescription(TArray<FString>& OutDescriptions, const TCHAR* External, const TCHAR* Internal)
 	{
 		if (!External) { return nullptr; }
@@ -156,7 +162,7 @@ namespace AsyncExecuteDetails
 			TArray<FString>& OutDescriptions,
 			const TCHAR* BaseDescription,
 			TFunction<bool(typename TAsyncResultCallbackArg<ResultType>::Type)>&& InFunc,
-			FTimespan Timeout = DefaultTimeout)
+			TOptional<FTimespan> Timeout = CQTest::DefaultTimeout)
 		{
 			return MakeShared<FWaitUntil>(
 				TestRunner,
@@ -179,7 +185,7 @@ public:
 	TAsyncExecute(
 		FAutomationTestBase& TestRunner,
 		TFunction<TAsyncResult<ResultType>()> AsyncAction,
-		FTimespan AsyncActionTimeout = AsyncExecuteDetails::DefaultTimeout,
+		TOptional<FTimespan> AsyncActionTimeout = CQTest::DefaultTimeout,
 		const TCHAR* Description = nullptr,
 		ResultCommandArgs&&... Args
 	)
diff --git a/Engine/Source/Developer/CQTest/Public/Components/MapTestSpawner.h b/Engine/Source/Developer/CQTest/Public/Components/MapTestSpawner.h
--- a/Engine/Source/Developer/CQTest/Public/Components/MapTestSpawner.h
+++ b/Engine/Source/Developer/CQTest/Public/Components/MapTestSpawner.h
@@ -73,10 +73,12 @@ struct CQTEST_API FMapTestSpawner : public FSpawnHelper
 	 * Loads the map specified from the MapDirectory and MapName to be prepared for the test.
 	 *
 	 * @param TestRunner - TestRunner used to send the latent command needed for map preparations.
+	 * @param Timeout - Duration that the latent command can execute before reporting an error.
 	 * 
 	 * @note Must be called outside of a latent action. Preferably within BEFORE_TEST.
+	 * @note If no value is provided for the timeout, value will be pulled from the `TestFramework.CQTest.CommandTimeout.MapTest` CVar.
 	 */
-	void AddWaitUntilLoadedCommand(FAutomationTestBase* TestRunner);
+	void AddWaitUntilLoadedCommand(FAutomationTestBase* TestRunner, TOptional<FTimespan> Timeout = CQTest::DefaultTimeout);
 
 	/**
 	 * Finds the first pawn in the given map.
diff --git a/Engine/Source/Developer/CQTest/Public/Components/PIENetworkComponent.h b/Engine/Source/Developer/CQTest/Public/Components/PIENetworkComponent.h
--- a/Engine/Source/Developer/CQTest/Public/Components/PIENetworkComponent.h
+++ b/Engine/Source/Developer/CQTest/Public/Components/PIENetworkComponent.h
@@ -8,6 +8,7 @@
 #if ENABLE_PIE_NETWORK_TEST
 
 #include "CQTest.h"
+#include "CQTestSettings.h"
 #include "PIENetworkTestStateRestorer.h"
 
 #include "Editor.h"
@@ -97,10 +98,12 @@ public:
 	 * @param InTestRunner - Pointer to the TestRunner used for test reporting.
 	 * @param InCommandBuilder - Reference to the latent command manager.
 	 * @param IsInitializing - Bool specifying if the Automation framework is still being initialized.
+	 * @param Timeout - Duration that the latent command can execute before reporting an error.
 	 * 
-	 * @note requires that the `FBasePIENetworkComponent` is built using the `FNetworkComponentBuilder` as this will setup the server and client `FBasePIENetworkComponentState`
+	 * @note Requires that the `FBasePIENetworkComponent` is built using the `FNetworkComponentBuilder` as this will setup the server and client `FBasePIENetworkComponentState`
+	 * @note If no value is provided for the timeout, value will be pulled from the `TestFramework.CQTest.CommandTimeout.Network` CVar.
 	 */
-	FBasePIENetworkComponent(FAutomationTestBase* InTestRunner, FTestCommandBuilder& InCommandBuilder, bool IsInitializing);
+	FBasePIENetworkComponent(FAutomationTestBase* InTestRunner, FTestCommandBuilder& InCommandBuilder, bool IsInitializing, TOptional<FTimespan> Timeout = CQTest::DefaultTimeout);
 
 	/**
 	 * Enqueues a function to the latent command manager for execution on a subsequent frame.
@@ -149,8 +152,9 @@ public:
 	 * @return a reference to this
 	 * 
 	 * @note Will continue to execute until the function provided evaluates to `true` or the duration exceeds the specified timeout.
+	 * @note If no value is provided for the timeout, value will be pulled from the `TestFramework.CQTest.CommandTimeout.Network` CVar.
 	 */
-	FBasePIENetworkComponent& Until(TFunction<bool()> Query, FTimespan Timeout = FTimespan::FromSeconds(10));
+	FBasePIENetworkComponent& Until(TFunction<bool()> Query, TOptional<FTimespan> Timeout = CQTest::DefaultTimeout);
 
 	/**
 	 * Enqueues a function to the latent command manager for execution until completion or timeout.
@@ -162,8 +166,9 @@ public:
 	 * @return a reference to this
 	 *
 	 * @note Will continue to execute until the function provided evaluates to `true` or the duration exceeds the specified timeout.
+	 * @note If no value is provided for the timeout, value will be pulled from the `TestFramework.CQTest.CommandTimeout.Network` CVar.
 	 */
-	FBasePIENetworkComponent& Until(const TCHAR* Description, TFunction<bool()> Query, FTimespan Timeout = FTimespan::FromSeconds(10));
+	FBasePIENetworkComponent& Until(const TCHAR* Description, TFunction<bool()> Query, TOptional<FTimespan> Timeout = CQTest::DefaultTimeout);
 	
 	/**
 	 * Enqueues a function to the latent command manager for execution until completion or timeout.
@@ -174,8 +179,9 @@ public:
 	 * @return a reference to this
 	 *
 	 * @note Will continue to execute until the function provided evaluates to `true` or the duration exceeds the specified timeout.
+	 * @note If no value is provided for the timeout, value will be pulled from the `TestFramework.CQTest.CommandTimeout.Network` CVar.
 	 */
-	FBasePIENetworkComponent& StartWhen(TFunction<bool()> Query, FTimespan Timeout = FTimespan::FromSeconds(10));
+	FBasePIENetworkComponent& StartWhen(TFunction<bool()> Query, TOptional<FTimespan> Timeout = CQTest::DefaultTimeout);
 
 	/**
 	 * Enqueues a function to the latent command manager for execution until completion or timeout.
@@ -187,8 +193,9 @@ public:
 	 * @return a reference to this
 	 *
 	 * @note Will continue to execute until the function provided evaluates to `true` or the duration exceeds the specified timeout.
+	 * @note If no value is provided for the timeout, value will be pulled from the `TestFramework.CQTest.CommandTimeout.Network` CVar.
 	 */
-	FBasePIENetworkComponent& StartWhen(const TCHAR* Description, TFunction<bool()> Query, FTimespan Timeout = FTimespan::FromSeconds(10));
+	FBasePIENetworkComponent& StartWhen(const TCHAR* Description, TFunction<bool()> Query, TOptional<FTimespan> Timeout = CQTest::DefaultTimeout);
 
 protected:
 	/** Prepares the active PIE sessions to be stopped. */
@@ -224,6 +231,15 @@ protected:
 	/** Restore back to the state prior to test execution. */
 	void RestoreState();
 
+	/**
+	 * Creates a timeout to be used.
+	 * 
+	 * @param Timeout - Duration that the latent command can execute before reporting an error.
+	 * 
+	 * @return either the FTimespan specified or the value from the `TestFramework.CQTest.CommandTimeout.Network` CVar.
+	 */
+	FTimespan MakeTimeout(TOptional<FTimespan> Timeout);
+
 	TUniquePtr<FBasePIENetworkComponentState> ServerState = nullptr;
 	TArray<TUniquePtr<FBasePIENetworkComponentState>> ClientStates;
 	FAutomationTestBase* TestRunner = nullptr;
@@ -323,8 +339,10 @@ public:
 	 * @param Timeout - Duration that the latent command can execute before reporting an error.
 	 * 
 	 * @return a reference to this
+	 * 
+	 * @note If no value is provided for the timeout, value will be pulled from the `TestFramework.CQTest.CommandTimeout.Network` CVar.
 	 */
-	FPIENetworkComponent& UntilServer(TFunction<bool(NetworkDataType&)> Query, FTimespan Timeout = FTimespan::FromSeconds(10));
+	FPIENetworkComponent& UntilServer(TFunction<bool(NetworkDataType&)> Query, TOptional<FTimespan> Timeout = CQTest::DefaultTimeout);
 
 	/**
 	 * Enqueues a function to the latent command manager for execution until completion or timeout on the server.
@@ -334,8 +352,10 @@ public:
 	 * @param Timeout - Duration that the latent command can execute before reporting an error.
 	 * 
 	 * @return a reference to this
+	 * 
+	 * @note If no value is provided for the timeout, value will be pulled from the `TestFramework.CQTest.CommandTimeout.Network` CVar.
 	 */
-	FPIENetworkComponent& UntilServer(const TCHAR* Description, TFunction<bool(NetworkDataType&)> Query, FTimespan Timeout = FTimespan::FromSeconds(10));
+	FPIENetworkComponent& UntilServer(const TCHAR* Description, TFunction<bool(NetworkDataType&)> Query, TOptional<FTimespan> Timeout = CQTest::DefaultTimeout);
 
 	/**
 	 * Enqueues a function to the latent command manager for execution until completion or timeout on all clients.
@@ -344,8 +364,10 @@ public:
 	 * @param Timeout - Duration that the latent command can execute before reporting an error.
 	 * 
 	 * @return a reference to this
+	 * 
+	 * @note If no value is provided for the timeout, value will be pulled from the `TestFramework.CQTest.CommandTimeout.Network` CVar.
 	 */
-	FPIENetworkComponent& UntilClients(TFunction<bool(NetworkDataType&)> Query, FTimespan Timeout = FTimespan::FromSeconds(10));
+	FPIENetworkComponent& UntilClients(TFunction<bool(NetworkDataType&)> Query, TOptional<FTimespan> Timeout = CQTest::DefaultTimeout);
 
 	/**
 	 * Enqueues a function to the latent command manager for execution until completion or timeout on all clients.
@@ -355,8 +377,10 @@ public:
 	 * @param Timeout - Duration that the latent command can execute before reporting an error.
 	 * 
 	 * @return a reference to this
+	 * 
+	 * @note If no value is provided for the timeout, value will be pulled from the `TestFramework.CQTest.CommandTimeout.Network` CVar.
 	 */
-	FPIENetworkComponent& UntilClients(const TCHAR* Description, TFunction<bool(NetworkDataType&)> Query, FTimespan Timeout = FTimespan::FromSeconds(10));
+	FPIENetworkComponent& UntilClients(const TCHAR* Description, TFunction<bool(NetworkDataType&)> Query, TOptional<FTimespan> Timeout = CQTest::DefaultTimeout);
 	
 	/**
 	 * Enqueues a function to the latent command manager for execution until completion or timeout on a given client.
@@ -366,8 +390,10 @@ public:
 	 * @param Timeout - Duration that the latent command can execute before reporting an error.
 	 * 
 	 * @return a reference to this
+	 * 
+	 * @note If no value is provided for the timeout, value will be pulled from the `TestFramework.CQTest.CommandTimeout.Network` CVar.
 	 */
-	FPIENetworkComponent& UntilClient(int32 ClientIndex, TFunction<bool(NetworkDataType&)> Query, FTimespan Timeout = FTimespan::FromSeconds(10));
+	FPIENetworkComponent& UntilClient(int32 ClientIndex, TFunction<bool(NetworkDataType&)> Query, TOptional<FTimespan> Timeout = CQTest::DefaultTimeout);
 
 	/**
 	 * Enqueues a function to the latent command manager for execution until completion or timeout on a given client.
@@ -378,38 +404,98 @@ public:
 	 * @param Timeout - Duration that the latent command can execute before reporting an error.
 	 * 
 	 * @return a reference to this
+	 * 
+	 * @note If no value is provided for the timeout, value will be pulled from the `TestFramework.CQTest.CommandTimeout.Network` CVar.
 	 */
-	FPIENetworkComponent& UntilClient(const TCHAR* Description, int32 ClientIndex, TFunction<bool(NetworkDataType&)> Query, FTimespan Timeout = FTimespan::FromSeconds(10));
+	FPIENetworkComponent& UntilClient(const TCHAR* Description, int32 ClientIndex, TFunction<bool(NetworkDataType&)> Query, TOptional<FTimespan> Timeout = CQTest::DefaultTimeout);
 
 	/**
 	 * Has a new client request a late join and syncs the server and other clients with the new client PIE session.
+	 * 
+	 * @param Timeout - Duration that the latent command can execute before reporting an error.
 	 *
 	 * @return a reference to this
+	 * 
+	 * @note If no value is provided for the timeout, value will be pulled from the `TestFramework.CQTest.CommandTimeout.Network` CVar.
 	 */
-	FPIENetworkComponent& ThenClientJoins();
+	FPIENetworkComponent& ThenClientJoins(TOptional<FTimespan> Timeout = CQTest::DefaultTimeout);
 	
 	/**
 	 * Spawns an actor on the server and replicates to the clients.
+	 * 
+	 * @param Timeout - Duration that the latent command can execute before reporting an error.
 	 *
 	 * @return a reference to this
+	 * 
+	 * @note If no value is provided for the timeout, value will be pulled from the `TestFramework.CQTest.CommandTimeout.Network` CVar.
 	 */
 	template<typename ActorToSpawn, ActorToSpawn* NetworkDataType::*ResultStorage>
-	FPIENetworkComponent& SpawnAndReplicate();
+	FPIENetworkComponent& SpawnAndReplicate(TOptional<FTimespan> Timeout = CQTest::DefaultTimeout);
 
+	/**
+	 * Spawns an actor on the server and replicates to the clients.
+	 *
+	 * @param SpawnParameters - Parameters that can be used when spawning an Actor.
+	 * @param Timeout - Duration that the latent command can execute before reporting an error.
+	 *
+	 * @return a reference to this
+	 *
+	 * @note If no value is provided for the timeout, value will be pulled from the `TestFramework.CQTest.CommandTimeout.Network` CVar.
+	 */
 	template<typename ActorToSpawn, ActorToSpawn* NetworkDataType::*ResultStorage>
-	FPIENetworkComponent& SpawnAndReplicate(const FActorSpawnParameters& SpawnParameters);
+	FPIENetworkComponent& SpawnAndReplicate(const FActorSpawnParameters& SpawnParameters, TOptional<FTimespan> Timeout = CQTest::DefaultTimeout);
 
+	/**
+	 * Spawns an actor on the server and replicates to the clients.
+	 *
+	 * @param BeforeReplicate - Function to be called after Actor has been spawned, but before being replicated to the other clients.
+	 * @param Timeout - Duration that the latent command can execute before reporting an error.
+	 *
+	 * @return a reference to this
+	 *
+	 * @note If no value is provided for the timeout, value will be pulled from the `TestFramework.CQTest.CommandTimeout.Network` CVar.
+	 */
 	template<typename ActorToSpawn, ActorToSpawn* NetworkDataType::* ResultStorage>
-	FPIENetworkComponent& SpawnAndReplicate(TFunction<void(ActorToSpawn&)> BeforeReplicate);
+	FPIENetworkComponent& SpawnAndReplicate(TFunction<void(ActorToSpawn&)> BeforeReplicate, TOptional<FTimespan> Timeout = CQTest::DefaultTimeout);
 
+	/**
+	 * Spawns an actor on the server and replicates to the connected clients.
+	 *
+	 * @param SpawnParameters - Parameters that can be used when spawning an Actor.
+	 * @param BeforeReplicate - Function to be called after Actor has been spawned, but before being replicated to the other clients.
+	 * @param Timeout - Duration that the latent command can execute before reporting an error.
+	 *
+	 * @return a reference to this
+	 *
+	 * @note If no value is provided for the timeout, value will be pulled from the `TestFramework.CQTest.CommandTimeout.Network` CVar.
+	 */
 	template<typename ActorToSpawn, ActorToSpawn* NetworkDataType::* ResultStorage>
-	FPIENetworkComponent& SpawnAndReplicate(const FActorSpawnParameters& SpawnParameters, TFunction<void(ActorToSpawn&)> BeforeReplicate);
+	FPIENetworkComponent& SpawnAndReplicate(const FActorSpawnParameters& SpawnParameters, TFunction<void(ActorToSpawn&)> BeforeReplicate, TOptional<FTimespan> Timeout = CQTest::DefaultTimeout);
 
 private:
 	friend class FNetworkComponentBuilder<NetworkDataType>;
+
+	/**
+	 * Spawns an actor on the server.
+	 *
+	 * @param SpawnParameters - Parameters that can be used when spawning an Actor.
+	 * @param BeforeReplicate - Function to be called after Actor has been spawned.
+	 * @param Timeout - Duration that the latent command can execute before reporting an error.
+	 *
+	 * @return a reference to this
+	 *
+	 * @note If no value is provided for the timeout, value will be pulled from the `TestFramework.CQTest.CommandTimeout.Network` CVar.
+	 */
 	template<typename ActorToSpawn, ActorToSpawn* NetworkDataType::* ResultStorage>
-	FPIENetworkComponent& SpawnOnServer(const FActorSpawnParameters& SpawnParameters, TFunction<void(ActorToSpawn&)> BeforeReplicate);
+	FPIENetworkComponent& SpawnOnServer(const FActorSpawnParameters& SpawnParameters, TFunction<void(ActorToSpawn&)> BeforeReplicate, TOptional<FTimespan> Timeout = CQTest::DefaultTimeout);
 
+	/**
+	 * Replicates all Actors to the connected client.
+	 *
+	 * @param ClientState - Client state to receive the updated Actor.
+	 *
+	 * @return true if all Actors have been successfully replicated on this client
+	 */
 	bool ReplicateToClients(NetworkDataType& ClientState);
 };
 
diff --git a/Engine/Source/Developer/CQTest/Public/Impl/PIENetworkComponent.inl b/Engine/Source/Developer/CQTest/Public/Impl/PIENetworkComponent.inl
--- a/Engine/Source/Developer/CQTest/Public/Impl/PIENetworkComponent.inl
+++ b/Engine/Source/Developer/CQTest/Public/Impl/PIENetworkComponent.inl
@@ -59,33 +59,35 @@ inline FPIENetworkComponent<NetworkDataType>& FPIENetworkComponent<NetworkDataTy
 }
 
 template <typename NetworkDataType>
-inline FPIENetworkComponent<NetworkDataType>& FPIENetworkComponent<NetworkDataType>::UntilServer(TFunction<bool(NetworkDataType&)> Query, FTimespan Timeout)
+inline FPIENetworkComponent<NetworkDataType>& FPIENetworkComponent<NetworkDataType>::UntilServer(TFunction<bool(NetworkDataType&)> Query, TOptional<FTimespan> Timeout)
 {
 	return UntilServer(nullptr, Query, Timeout);
 }
 
 template <typename NetworkDataType>
-inline FPIENetworkComponent<NetworkDataType>& FPIENetworkComponent<NetworkDataType>::UntilServer(const TCHAR* Description, TFunction<bool(NetworkDataType&)> Query, FTimespan Timeout)
+inline FPIENetworkComponent<NetworkDataType>& FPIENetworkComponent<NetworkDataType>::UntilServer(const TCHAR* Description, TFunction<bool(NetworkDataType&)> Query, TOptional<FTimespan> Timeout)
 {
+	FTimespan TimeoutValue = MakeTimeout(Timeout);
 	CommandBuilder->Until(Description, [this, Query]() { 
 		return Query(static_cast<NetworkDataType&>(*ServerState)); 
-	}, Timeout);
+	}, TimeoutValue);
 
 	return *this;
 }
 
 template <typename NetworkDataType>
-inline FPIENetworkComponent<NetworkDataType>& FPIENetworkComponent<NetworkDataType>::UntilClients(TFunction<bool(NetworkDataType&)> Query, FTimespan Timeout)
+inline FPIENetworkComponent<NetworkDataType>& FPIENetworkComponent<NetworkDataType>::UntilClients(TFunction<bool(NetworkDataType&)> Query, TOptional<FTimespan> Timeout)
 {
 	return UntilClients(nullptr, Query, Timeout);
 }
 
 template <typename NetworkDataType>
-inline FPIENetworkComponent<NetworkDataType>& FPIENetworkComponent<NetworkDataType>::UntilClients(const TCHAR* Description, TFunction<bool(NetworkDataType&)> Query, FTimespan Timeout)
+inline FPIENetworkComponent<NetworkDataType>& FPIENetworkComponent<NetworkDataType>::UntilClients(const TCHAR* Description, TFunction<bool(NetworkDataType&)> Query, TOptional<FTimespan> Timeout)
 {
 	// Capture a mutable array of bools (by value) to track which clients have already finished to avoid calling them again
+	FTimespan TimeoutValue = MakeTimeout(Timeout);
 	TArray<bool> ClientsFinishedTask{};
-	CommandBuilder->Until(Description, [this, Query, Timeout, ClientsFinishedTask]() mutable {
+	CommandBuilder->Until(Description, [this, Query, ClientsFinishedTask]() mutable {
 		// Array with be initially empty and will need to be resized to match our current client count
 		// It is safe to assume that no clients will join in the middle of this action so the resize will occur only once
 		if (ClientsFinishedTask.Num() < ClientStates.Num())
@@ -110,19 +112,21 @@ inline FPIENetworkComponent<NetworkDataType>& FPIENetworkComponent<NetworkDataTy
 		}
 
 		return bIsAllDone;
-	});
+	}, TimeoutValue);
 
 	return *this;
 }
+
 template <typename NetworkDataType>
-inline FPIENetworkComponent<NetworkDataType>& FPIENetworkComponent<NetworkDataType>::UntilClient(int32 ClientIndex, TFunction<bool(NetworkDataType&)> Query, FTimespan Timeout)
+inline FPIENetworkComponent<NetworkDataType>& FPIENetworkComponent<NetworkDataType>::UntilClient(int32 ClientIndex, TFunction<bool(NetworkDataType&)> Query, TOptional<FTimespan> Timeout)
 {
 	return UntilClient(nullptr, ClientIndex, Query, Timeout);
 }
 
 template <typename NetworkDataType>
-inline FPIENetworkComponent<NetworkDataType>& FPIENetworkComponent<NetworkDataType>::UntilClient(const TCHAR* Description, int32 ClientIndex, TFunction<bool(NetworkDataType&)> Query, FTimespan Timeout)
+inline FPIENetworkComponent<NetworkDataType>& FPIENetworkComponent<NetworkDataType>::UntilClient(const TCHAR* Description, int32 ClientIndex, TFunction<bool(NetworkDataType&)> Query, TOptional<FTimespan> Timeout)
 {
+	FTimespan TimeoutValue = MakeTimeout(Timeout);
 	CommandBuilder->Until(Description, [this, Query, ClientIndex]() {
 		if (!ClientStates.IsValidIndex(ClientIndex))
 		{
@@ -130,14 +134,15 @@ inline FPIENetworkComponent<NetworkDataType>& FPIENetworkComponent<NetworkDataTy
 			return true;
 		}
 		return Query(static_cast<NetworkDataType&>(*ClientStates[ClientIndex]));
-	}, Timeout);
+	}, TimeoutValue);
 
 	return *this;
 }
 
 template <typename NetworkDataType>
-inline FPIENetworkComponent<NetworkDataType>& FPIENetworkComponent<NetworkDataType>::ThenClientJoins()
+inline FPIENetworkComponent<NetworkDataType>& FPIENetworkComponent<NetworkDataType>::ThenClientJoins(TOptional<FTimespan> Timeout)
 {
+	FTimespan TimeoutValue = MakeTimeout(Timeout);
 	Do(TEXT("Update Server State"), [this]() {
 		int32 NextClientIndex = ServerState->ClientCount++;
 		ClientStates.Add(MakeUnique<NetworkDataType>(NetworkDataType{}));
@@ -145,52 +150,54 @@ inline FPIENetworkComponent<NetworkDataType>& FPIENetworkComponent<NetworkDataTy
 		ServerState->ClientConnections.SetNum(ServerState->ClientCount);
 		GEditor->RequestLateJoin();
 	})
-	.Until(TEXT("Setting Worlds"), [this]() { return SetWorlds(); })
+	.Until(TEXT("Setting Worlds"), [this]() { return SetWorlds(); }, TimeoutValue)
 	.Then(TEXT("Setup Packet Settings"), [this]() { SetPacketSettings(); })
 	.Then(TEXT("Connect Clients to Server"), [this]() { ConnectClientsToServer(); });
 
 	UntilClient(TEXT("Replicate to new Client"), ServerState->ClientCount, [this](NetworkDataType& State) {
 		return ReplicateToClients(State);
-	});
+	}, TimeoutValue);
 	
 	return *this;
 }
 
 template <typename NetworkDataType>
 template <typename ActorToSpawn, ActorToSpawn* NetworkDataType::*ResultStorage>
-inline FPIENetworkComponent<NetworkDataType>& FPIENetworkComponent<NetworkDataType>::SpawnAndReplicate()
+inline FPIENetworkComponent<NetworkDataType>& FPIENetworkComponent<NetworkDataType>::SpawnAndReplicate(TOptional<FTimespan> Timeout)
 {
-	return SpawnAndReplicate<ActorToSpawn, ResultStorage>({}, {});
+	return SpawnAndReplicate<ActorToSpawn, ResultStorage>({}, {}, Timeout);
 }
 
 template <typename NetworkDataType>
 template<typename ActorToSpawn, ActorToSpawn* NetworkDataType::* ResultStorage>
-inline FPIENetworkComponent<NetworkDataType>& FPIENetworkComponent<NetworkDataType>::SpawnAndReplicate(const FActorSpawnParameters& SpawnParameters) {
-	return SpawnAndReplicate<ActorToSpawn, ResultStorage>(SpawnParameters, {});
+inline FPIENetworkComponent<NetworkDataType>& FPIENetworkComponent<NetworkDataType>::SpawnAndReplicate(const FActorSpawnParameters& SpawnParameters, TOptional<FTimespan> Timeout) {
+	return SpawnAndReplicate<ActorToSpawn, ResultStorage>(SpawnParameters, {}, Timeout);
 }
 
 template <typename NetworkDataType>
 template<typename ActorToSpawn, ActorToSpawn* NetworkDataType::* ResultStorage>
-inline FPIENetworkComponent<NetworkDataType>& FPIENetworkComponent<NetworkDataType>::SpawnAndReplicate(TFunction<void(ActorToSpawn&)> BeforeReplicate) {
-	return SpawnAndReplicate<ActorToSpawn, ResultStorage>({}, BeforeReplicate);
+inline FPIENetworkComponent<NetworkDataType>& FPIENetworkComponent<NetworkDataType>::SpawnAndReplicate(TFunction<void(ActorToSpawn&)> BeforeReplicate, TOptional<FTimespan> Timeout) {
+	return SpawnAndReplicate<ActorToSpawn, ResultStorage>({}, BeforeReplicate, Timeout);
 }
 
 template <typename NetworkDataType>
 template<typename ActorToSpawn, ActorToSpawn* NetworkDataType::* ResultStorage>
-inline FPIENetworkComponent<NetworkDataType>& FPIENetworkComponent<NetworkDataType>::SpawnAndReplicate(const FActorSpawnParameters& SpawnParameters, TFunction<void(ActorToSpawn&)> BeforeReplicate) 
+inline FPIENetworkComponent<NetworkDataType>& FPIENetworkComponent<NetworkDataType>::SpawnAndReplicate(const FActorSpawnParameters& SpawnParameters, TFunction<void(ActorToSpawn&)> BeforeReplicate, TOptional<FTimespan> Timeout)
 {
-	SpawnOnServer<ActorToSpawn, ResultStorage>(SpawnParameters, BeforeReplicate);
-	UntilClients([this](NetworkDataType& ClientState) { return ReplicateToClients(ClientState); });
+	FTimespan TimeoutValue = MakeTimeout(Timeout);
+	SpawnOnServer<ActorToSpawn, ResultStorage>(SpawnParameters, BeforeReplicate, TimeoutValue);
+	UntilClients([this](NetworkDataType& ClientState) { return ReplicateToClients(ClientState); }, TimeoutValue);
 		
 	return *this;
 }
 
 template <typename NetworkDataType>
 template<typename ActorToSpawn, ActorToSpawn* NetworkDataType::* ResultStorage>
-inline FPIENetworkComponent<NetworkDataType>& FPIENetworkComponent<NetworkDataType>::SpawnOnServer(const FActorSpawnParameters& SpawnParameters, TFunction<void(ActorToSpawn&)> BeforeReplicate) {
+inline FPIENetworkComponent<NetworkDataType>& FPIENetworkComponent<NetworkDataType>::SpawnOnServer(const FActorSpawnParameters& SpawnParameters, TFunction<void(ActorToSpawn&)> BeforeReplicate, TOptional<FTimespan> Timeout) {
 	static_assert(std::is_convertible_v<ActorToSpawn*, AActor*>, "ActorToSpawn must derive from AActor");
 	static_assert(std::is_default_constructible<NetworkDataType>::value, "NetworkDataType must have a default constructor accessible");
 
+	FTimespan TimeoutValue = MakeTimeout(Timeout);
 	// TODO: Consider passing in a constructed actor from an FTestSpawner instead
 	// That would allow using TObjectBuilder as well.
 	// Need a version which takes the Server's world instead of creating its own
@@ -226,7 +233,7 @@ inline FPIENetworkComponent<NetworkDataType>& FPIENetworkComponent<NetworkDataTy
 		SpawnedActors.Add(NetGUID, StorageOffset);
 		State.LocallySpawnedActors.Add(NetGUID);
 		return true;
-	});
+	}, TimeoutValue);
 	
 	return *this;
 }
diff --git a/Engine/Source/Developer/CQTest/README.md b/Engine/Source/Developer/CQTest/README.md
--- a/Engine/Source/Developer/CQTest/README.md
+++ b/Engine/Source/Developer/CQTest/README.md
@@ -309,6 +309,16 @@ Instead, it is better to add the actions as a series of self-contained steps.
 
 The framework has been designed to allow for extensions in a couple areas.  In-code examples can be found within the CQTest plugin used to test the framework. See _/Engine/Plugins/Tests/CQTest/Source/CQTestTests/Private/ExtensionTests.cpp_ for in-code examples.
 
+## Test Settings
+
+The framework also comes with settings that can be applied per project.  The settings are saved to the project's Engine configuration file located in _Config/DefaultEngine.ini_ under the section `[/Script/CQTest.CQTestSettings]`. Some of the available settings are:
+- `TestFramework.CQTest.CommandTimeout` - Timeout value, in seconds, that is applied to commands that wait on an action to be evaluated to true and no user defined timeout is specified.
+- `TestFramework.CQTest.CommandTimeout.Network` - Timeout value, in seconds, that is applied to commands that wait on a PIENetworkComponent action to be evaluated to true and no user defined timeout is specified.
+- `TestFramework.CQTest.CommandTimeout.MapTest` - Timeout value, in seconds, that is applied to the `FMapTestSpawner::AddWaitUntilLoadedCommand` method for loading a map and starting a PIE session.
+
+The above mentioned settings are stored as Console Variables and can also be adjusted within the Editor or Game build targets through the console.
+**NOTE:** Adjusting the variables through the console will only have the value applied through the duration of the application run and will reset to the defaults if no prior saved settings exist upon exit.  Changes must be modified through the Editor by going into the **Project Settings** and then finding **CQ Test Settings** under the **Engine** category.
+
 ## Test Components
 
 This testing framework embraces composition over inheritence.  Creating new components should be the default mechanism for extending the framework.  Some of the components available to you are:
@@ -324,6 +334,9 @@ This testing framework embraces composition over inheritence.  Creating new comp
 
 ### Notable Changes
 
+#### 5.6
+- Default timeout durations are able to be configured per project via the **Project Settings** within the Editor or through CVars.
+   
 #### 5.5
 - **CQTest** - The core **CQTest** framework has been extracted out from the plugin and is now an Engine Module. While the plugin still exists and is used to test that the core functionality works; it is deprecated and not necessary to be included as part of the project. No action needs to be taken for existing projects to have the core **CQTest** framework available.
 - **CQTestEnhancedInput** - **BREAKING CHANGE** Due to the move of **CQTest** being a module, the components that used the **EnhancedInput** plugin had to be extracted out into a separate plugin. **EnhancedInput** is currently an Engine plugin, similar to what **CQTest** used to be. Because both were considered Engine plugins, they were able to reference each other. With the move of **CQTest** being in the Engine and cannot reference an Engine plugin, there is a need to add `CQTestEnhancedInput` to the project's Build.cs file, similar to how `CQTest` was added. This will only impact the project if the `InputTestAction` component was being used, but does not have an impact on the core **CQTest** framework.
