Addressing Linux ASAN new-delete-type-mismatch error.

#rnx

diff --git a/Engine/Source/Developer/CQTest/Public/Components/PIENetworkComponent.h b/Engine/Source/Developer/CQTest/Public/Components/PIENetworkComponent.h
--- a/Engine/Source/Developer/CQTest/Public/Components/PIENetworkComponent.h
+++ b/Engine/Source/Developer/CQTest/Public/Components/PIENetworkComponent.h
@@ -66,6 +66,9 @@ NETWORK_TEST_CLASS(MyFixtureName, "Network.Example")
  */
 struct FBasePIENetworkComponentState
 {
+	/** Default virtual destructor. */
+	virtual ~FBasePIENetworkComponentState() = default;
+
 	/** Reference to the session's world. */
 	UWorld* World = nullptr;
 
