Fixing typo in _TEST_CLASS_IMPL macro

#rnx

diff --git a/Engine/Source/Developer/CQTest/README.md b/Engine/Source/Developer/CQTest/README.md
--- a/Engine/Source/Developer/CQTest/README.md
+++ b/Engine/Source/Developer/CQTest/README.md
@@ -185,7 +185,7 @@ In addition to `TEST` and `TEST_CLASS` are 5 additional macros:
  - `TEST_CLASS_WITH_BASE` - Macro which allows this test object to inherit from a different test object.  More information about how to use this macro can be found [in the section below regarding custom test classes](#base-test-class)
  - `TEST_CLASS_WITH_FLAGS` - Macro which allows the use of different automation test flags to be specified.  Useful for when tests can only run under a certain context or grouped under a specific filter.  The default flags for `TEST` and `TEST_CLASS` are `EAutomationTestFlags_ApplicationContextMask | EAutomationTestFlags::ProductFilter`.  Additional information regarding the available flags can be found [in the online documentation](https://dev.epicgames.com/documentation/en-us/unreal-engine/API/Runtime/Core/Misc/EAutomationTestFlags/Type)
  - `TEST_CLASS_WITH_BASE_AND_FLAGS` - Macro which allows for this test object to inherit from a different test object and allows for custom automation test flags to be specified.
- - `TEST_CLASS_IMPL` - Base macro which is used by the above macros to specify a custom asserter, a test object to inherit from a different test object, and allows for custom automation test flags to be specified.
+ - `_TEST_CLASS_IMPL` - Base macro which is used by the above macros to specify a custom asserter, a test object to inherit from a different test object, and allows for custom automation test flags to be specified.
 
 Test Directory determines where in the Automation tab the tests appear.  In the example above, we specify "Game.MyGame", but you may also have an auto-generated test directory based on the folder structure.
 ```cpp
