Fixing test names being appended to the list everytime the tests are executed.

#rnx

diff --git a/Engine/Source/Developer/CQTest/Public/CQTest.h b/Engine/Source/Developer/CQTest/Public/CQTest.h
--- a/Engine/Source/Developer/CQTest/Public/CQTest.h
+++ b/Engine/Source/Developer/CQTest/Public/CQTest.h
@@ -208,9 +208,13 @@ struct TTest : TBaseTest<AsserterType>
 	{
 		FFunctionRegistrar(FString Name, TestMethod Func, int32 LineNumber)
 		{
-			DerivedType::TestRunner->TestNames.Add(Name);
-			DerivedType::Methods.Add(Name, Func);
-			DerivedType::TestRunner->TestLineNumbers.Add(Name, LineNumber);
+			// Only register the test method if not already registered
+			if (!DerivedType::Methods.Contains(Name))
+			{
+				DerivedType::TestRunner->TestNames.Add(Name);
+				DerivedType::Methods.Add(Name, Func);
+				DerivedType::TestRunner->TestLineNumbers.Add(Name, LineNumber);
+			}
 		}
 	};
 
