CQ Tests - tagging support

diff --git a/Engine/Plugins/Tests/CQTest/Source/CQTestTests/Private/CQTestTests.cpp b/Engine/Plugins/Tests/CQTest/Source/CQTestTests/Private/CQTestTests.cpp
--- a/Engine/Plugins/Tests/CQTest/Source/CQTestTests/Private/CQTestTests.cpp
+++ b/Engine/Plugins/Tests/CQTest/Source/CQTestTests/Private/CQTestTests.cpp
@@ -76,6 +76,27 @@ namespace CQTestTests
 		}
 	};
 
+	TEST_CLASS_WITH_FLAGS_AND_TAGS(OverrideFixtureTestFlagsAndTags, "TestFramework.CQTest.Core.Tags", EAutomationTestFlags_ApplicationContextMask | EAutomationTestFlags::EngineFilter, "[CQFlagsTest][FlagsExtraTag][CQTaggedTests]")
+	{
+		TEST_METHOD(GetTestFlags_ReturnsSetAutomationTestFlags_HasExpectedTags)
+		{
+			ASSERT_THAT(AreEqual(EAutomationTestFlags_ApplicationContextMask | EAutomationTestFlags::EngineFilter, TestRunner->GetTestFlags()));
+
+			const FString TestTags = TestRunner->GetTestTags();
+			ASSERT_THAT(IsTrue(TestTags.Contains(TEXT("[CQFlagsTest]"))));
+			ASSERT_THAT(IsTrue(TestTags.Contains(TEXT("[FlagsExtraTag]"))));
+		}
+
+		TEST_METHOD_WITH_TAGS(GetTestFlags_HasExpectedMethodTags, "[CQClassWithFlagsAndTagsMethod]")
+		{
+			const FString ThisMethodName = TEXT("GetTestFlags_HasExpectedMethodTags");
+			const FString TestTags = TestRunner->GetTestMethodTags(ThisMethodName);
+			ASSERT_THAT(IsTrue(TestTags.Contains(TEXT("[CQFlagsTest]"))));
+			ASSERT_THAT(IsTrue(TestTags.Contains(TEXT("[FlagsExtraTag]"))));
+			ASSERT_THAT(IsTrue(TestTags.Contains(TEXT("[CQClassWithFlagsAndTagsMethod]"))));
+		}
+	};
+
 	TEST_CLASS(TestFixtureTest, "TestFramework.CQTest.Core")
 	{
 		bool SetupCalled = false;
@@ -591,4 +612,31 @@ namespace CQTestTests
 			ASSERT_THAT(IsTrue(DoesEventExist(*this->TestRunner, FAutomationEvent(EAutomationEventType::Error, ExpectedMessage))));
 		}
 	};
+
+	TEST_CLASS_WITH_TAGS(CQTags, "TestFramework.CQTest.Core.Tags", "[CQTest][ExtraTag][CQTaggedTests]")
+	{
+		TEST_METHOD(CheckTagsAreValid)
+		{
+			const FString TestTags = TestRunner->GetTestTags();
+			ASSERT_THAT(IsTrue(TestTags.Contains(TEXT("[CQTest]"))));
+			ASSERT_THAT(IsTrue(TestTags.Contains(TEXT("[ExtraTag]"))));
+		}
+
+		TEST_METHOD_WITH_TAGS(CheckTagsAreValid_HasExpectedMethodTags, "[CQClassWithTagsMethod]")
+		{
+			const FString ThisMethodName = TEXT("CheckTagsAreValid_HasExpectedMethodTags");
+			const FString TestTags = TestRunner->GetTestMethodTags(ThisMethodName);
+			ASSERT_THAT(IsTrue(TestTags.Contains(TEXT("[CQTest]"))));
+			ASSERT_THAT(IsTrue(TestTags.Contains(TEXT("[ExtraTag]"))));
+			ASSERT_THAT(IsTrue(TestTags.Contains(TEXT("[CQClassWithTagsMethod]"))));
+		}
+
+		TEST_METHOD_WITH_TAGS(CheckTagsAreValid_HasExpectedMethodTagsNoLeak, "[CQClassWithTagsNoLeakMethod]")
+		{
+			const FString ThisMethodName = TEXT("CheckTagsAreValid_HasExpectedMethodTagsNoLeak");
+			const FString TestTags = TestRunner->GetTestMethodTags(ThisMethodName);
+			ASSERT_THAT(IsTrue(TestTags.Contains(TEXT("[CQClassWithTagsNoLeakMethod]"))));
+			ASSERT_THAT(IsFalse(TestTags.Contains(TEXT("[CQClassWithTagsMethod]"))));
+		}
+	};
 } // namespace CQTestTests
diff --git a/Engine/Plugins/Tests/CQTest/Source/CQTestTests/Private/ExtensionTests.cpp b/Engine/Plugins/Tests/CQTest/Source/CQTestTests/Private/ExtensionTests.cpp
--- a/Engine/Plugins/Tests/CQTest/Source/CQTestTests/Private/ExtensionTests.cpp
+++ b/Engine/Plugins/Tests/CQTest/Source/CQTestTests/Private/ExtensionTests.cpp
@@ -28,6 +28,27 @@ CUSTOM_ASSERT_TEST_CLASS(CustomAsserts, "TestFramework.CQTest.Core")
 	}
 };
 
+TEST_CLASS_WITH_ASSERTS_AND_TAGS(CustomAssertsTags, "TestFramework.CQTest.Core.Tags", FCustomAsserter, "[CQAssertTest][AssertExtraTag][CQTaggedTests]")
+{
+	TEST_METHOD(CustomTestClass_WithCustomAsserter_UsesCustomAsserter_HasExpectedTags)
+	{
+		ASSERT_THAT(Custom(true));
+		
+		const FString TestTags = TestRunner->GetTestTags();
+		ASSERT_THAT(IsTrue(TestTags.Contains(TEXT("[CQAssertTest]"))));
+		ASSERT_THAT(IsTrue(TestTags.Contains(TEXT("[AssertExtraTag]"))));
+	}
+
+	TEST_METHOD_WITH_TAGS(CustomTestClass_HasExpectedMethodTags, "[CQClassWithAssertsAndTagsMethod]")
+	{
+		const FString ThisMethodName = TEXT("CustomTestClass_HasExpectedMethodTags");
+		const FString TestTags = TestRunner->GetTestMethodTags(ThisMethodName);
+		ASSERT_THAT(IsTrue(TestTags.Contains(TEXT("[CQAssertTest]"))));
+		ASSERT_THAT(IsTrue(TestTags.Contains(TEXT("[AssertExtraTag]"))));
+		ASSERT_THAT(IsTrue(TestTags.Contains(TEXT("[CQClassWithAssertsAndTagsMethod]"))));
+	}
+};
+
 template<typename Derived, typename AsserterType>
 struct TCustomBaseClass : public TTest<Derived, AsserterType>
 {
@@ -75,6 +96,60 @@ CUSTOM_BASE_TEST_CLASS(DerivedTest, "TestFramework.CQTest.Core")
 	}
 };
 
+TEST_CLASS_WITH_BASE_AND_FLAGS(DerivedWithFlagsTest, "TestFramework.CQTest.Core", TCustomBaseClass, EAutomationTestFlags_ApplicationContextMask | EAutomationTestFlags::EngineFilter)
+{
+	TEST_METHOD(DerivedTestClass_WithCustomBase_InheritsFromBaseClass_HasExpectedFlags)
+	{
+		ASSERT_THAT(AreEqual(42, SpecialValue));
+
+		ASSERT_THAT(AreEqual(EAutomationTestFlags_ApplicationContextMask | EAutomationTestFlags::EngineFilter, TestRunner->GetTestFlags()));
+	}
+};
+
+TEST_CLASS_WITH_BASE_AND_TAGS(DerivedWithTagsTest, "TestFramework.CQTest.Core.Tags", TCustomBaseClass, "[CQBaseTest][BaseExtraTag][CQTaggedTests]")
+{
+	TEST_METHOD(DerivedTestClass_WithCustomBase_InheritsFromBaseClass_HasExpectedTags)
+	{
+		ASSERT_THAT(AreEqual(42, SpecialValue));
+
+		const FString TestTags = TestRunner->GetTestTags();
+		ASSERT_THAT(IsTrue(TestTags.Contains(TEXT("[CQBaseTest]"))));
+		ASSERT_THAT(IsTrue(TestTags.Contains(TEXT("[BaseExtraTag]"))));
+	}
+
+	TEST_METHOD_WITH_TAGS(DerivedTestClass_HasExpectedMethodTags, "[CQTestClassWithBaseAndTagsMethod]")
+	{
+		const FString ThisMethodName = TEXT("DerivedTestClass_HasExpectedMethodTags");
+		const FString TestTags = TestRunner->GetTestMethodTags(ThisMethodName);
+		ASSERT_THAT(IsTrue(TestTags.Contains(TEXT("[CQBaseTest]"))));
+		ASSERT_THAT(IsTrue(TestTags.Contains(TEXT("[BaseExtraTag]"))));
+		ASSERT_THAT(IsTrue(TestTags.Contains(TEXT("[CQTestClassWithBaseAndTagsMethod]"))));
+	}
+};
+
+TEST_CLASS_WITH_BASE_AND_FLAGS_AND_TAGS(DerivedWithFlagsAndTagsTest, "TestFramework.CQTest.Core.Tags", TCustomBaseClass, EAutomationTestFlags_ApplicationContextMask | EAutomationTestFlags::EngineFilter, "[CQBaseFlagsTest][BaseFlagsExtraTag][CQTaggedTests]")
+{
+	TEST_METHOD(DerivedTestClass_WithCustomBase_InheritsFromBaseClass_HasExpectedFlags_HasExpectedTags)
+	{
+		ASSERT_THAT(AreEqual(42, SpecialValue));
+
+		ASSERT_THAT(AreEqual(EAutomationTestFlags_ApplicationContextMask | EAutomationTestFlags::EngineFilter, TestRunner->GetTestFlags()));
+
+		const FString TestTags = TestRunner->GetTestTags();
+		ASSERT_THAT(IsTrue(TestTags.Contains(TEXT("[CQBaseFlagsTest]"))));
+		ASSERT_THAT(IsTrue(TestTags.Contains(TEXT("[BaseFlagsExtraTag]"))));
+	}
+
+	TEST_METHOD_WITH_TAGS(DerivedTestClass_HasExpectedMethodTags, "[CQTestClassWithBaseAndFlagsAndTagsMethod]")
+	{
+		const FString ThisMethodName = TEXT("DerivedTestClass_HasExpectedMethodTags");
+		const FString TestTags = TestRunner->GetTestMethodTags(ThisMethodName);
+		ASSERT_THAT(IsTrue(TestTags.Contains(TEXT("[CQBaseFlagsTest]"))));
+		ASSERT_THAT(IsTrue(TestTags.Contains(TEXT("[BaseFlagsExtraTag]"))));
+		ASSERT_THAT(IsTrue(TestTags.Contains(TEXT("[CQTestClassWithBaseAndFlagsAndTagsMethod]"))));
+	}
+};
+
 template <typename Derived, typename AsserterType>
 struct TBaseWithConstructor : public TTest<Derived, AsserterType>
 {
diff --git a/Engine/Source/Developer/CQTest/Public/CQTest.h b/Engine/Source/Developer/CQTest/Public/CQTest.h
--- a/Engine/Source/Developer/CQTest/Public/CQTest.h
+++ b/Engine/Source/Developer/CQTest/Public/CQTest.h
@@ -66,6 +66,7 @@ enum class ECQTestSuppressLogBehavior
 
 static const FString GenerateTestDirectory = TEXT("");
 static constexpr EAutomationTestFlags DefaultFlags = EAutomationTestFlags_ApplicationContextMask | EAutomationTestFlags::ProductFilter;
+static const FString DefaultTags = TEXT("");
 
 template <typename T>
 concept HasBeforeAll = requires(T t) { { T::BeforeAll(FString()) }; };
@@ -143,7 +144,7 @@ using TTestInstanceGenerator = TUniquePtr<TBaseTest<AsserterType>>(*)(TTestRunne
 template <typename AsserterType>
 struct TTestRunner : public FAutomationTestBase
 {
-	TTestRunner(FString Name, int32 LineNumber, const char* FileName, FString TestDir, EAutomationTestFlags TestFlags, TTestInstanceGenerator<AsserterType> Factory);
+	TTestRunner(FString Name, int32 LineNumber, const char* FileName, FString TestDir, EAutomationTestFlags TestFlags, TTestInstanceGenerator<AsserterType> Factory, FString TestTags);
 	~TTestRunner();
 
 	EAutomationTestFlags GetTestFlags() const override;
@@ -159,6 +160,15 @@ struct TTestRunner : public FAutomationTestBase
 	virtual void SetSuppressLogWarnings(ECQTestSuppressLogBehavior Behavior = ECQTestSuppressLogBehavior::True);
 	virtual void SetSuppressLogErrors(ECQTestSuppressLogBehavior Behavior = ECQTestSuppressLogBehavior::True);
 
+	const FString GetTestTags() const { return TestTags; }
+	const FString GetTestMethodTags(const FString& MethodName) const
+	{
+		FString FullMethodName = FString::Printf(TEXT("%s.%s"), *GetBeautifiedTestName(), *MethodName);
+		return FAutomationTestFramework::Get().GetTagsForAutomationTest(FullMethodName);
+	}
+
+	FString GetBeautifiedTestName() const override;
+
 	int32 LineNumber = 0;
 	FString FileName;
 	FString TestDir;
@@ -175,12 +185,13 @@ struct TTestRunner : public FAutomationTestBase
 	TTestRunner() = delete;
 
 protected:
-	FString GetBeautifiedTestName() const override;
 	uint32 GetRequiredDeviceNum() const override;
 
 	ECQTestSuppressLogBehavior SuppressLogWarningsBehavior{ ECQTestSuppressLogBehavior::Default };
 	ECQTestSuppressLogBehavior SuppressLogErrorsBehavior{ ECQTestSuppressLogBehavior::Default };
 
+private:
+	FString TestTags;
 };
 
 template <typename Derived, typename AsserterType>
@@ -206,7 +217,7 @@ struct TTest : TBaseTest<AsserterType>
 
 	struct FFunctionRegistrar
 	{
-		FFunctionRegistrar(FString Name, TestMethod Func, int32 LineNumber)
+		FFunctionRegistrar(FString Name, TestMethod Func, int32 LineNumber, FString TestTags)
 		{
 			// Only register the test method if not already registered
 			if (!DerivedType::Methods.Contains(Name))
@@ -214,8 +225,24 @@ struct TTest : TBaseTest<AsserterType>
 				DerivedType::TestRunner->TestNames.Add(Name);
 				DerivedType::Methods.Add(Name, Func);
 				DerivedType::TestRunner->TestLineNumbers.Add(Name, LineNumber);
+
+				FString CommonTestTags = DerivedType::TestRunner->GetTestTags();
+				FString CompleteTags = FString::Printf(TEXT("%s%s"), *CommonTestTags, *TestTags);
+				if (!CompleteTags.IsEmpty())
+				{
+					FString FullName = FString::Printf(TEXT("%s.%s"), *DerivedType::TestRunner->GetBeautifiedTestName(), *Name);
+					bool Registered = FAutomationTestFramework::Get().RegisterAutomationTestTags(FullName, CompleteTags);
+					if (Registered)
+					{
+						MethodTags = CompleteTags;
+					}
+				}
 			}
 		}
+
+		const FString GetTestMethodTags() const { return MethodTags; }
+	private:
+		FString MethodTags;
 	};
 
 	static TUniquePtr<TBaseTest<AsserterType>> CreateTestClass(TTestRunner<AsserterType>& TestRunnerInstance)
@@ -229,42 +256,56 @@ struct TTest : TBaseTest<AsserterType>
 };
 
 #if WITH_AUTOMATION_WORKER
-#define _TEST_CLASS_IMPL(_ClassName, _TestDir, _BaseClass, _AsserterType, _TestFlags)                                                  \
-	struct _ClassName;                                                                                                                 \
-	struct F##_ClassName##_Runner : public TTestRunner<_AsserterType>                                                                  \
-	{                                                                                                                                  \
-		F##_ClassName##_Runner()                                                                                                       \
-			: TTestRunner(#_ClassName, __LINE__, __FILE__, _TestDir, _TestFlags, TTest<_ClassName, _AsserterType>::CreateTestClass) {  \
-				static_assert(  !!((_TestFlags) & EAutomationTestFlags_ApplicationContextMask),                                        \
-								"CQTest has no application flag and will not run. See AutomationTest.h.");                             \
-				static_assert(  !!(((_TestFlags) & EAutomationTestFlags_FilterMask) == EAutomationTestFlags::SmokeFilter) ||           \
-								!!(((_TestFlags) & EAutomationTestFlags_FilterMask) == EAutomationTestFlags::EngineFilter) ||          \
-								!!(((_TestFlags) & EAutomationTestFlags_FilterMask) == EAutomationTestFlags::ProductFilter) ||         \
-								!!(((_TestFlags) & EAutomationTestFlags_FilterMask) == EAutomationTestFlags::PerfFilter) ||            \
-								!!(((_TestFlags) & EAutomationTestFlags_FilterMask) == EAutomationTestFlags::StressFilter) ||          \
-								!!(((_TestFlags) & EAutomationTestFlags_FilterMask) == EAutomationTestFlags::NegativeFilter),          \
-								"All CQTests must have exactly 1 filter type specified. See AutomationTest.h.");                       \
-		}                                                                                                                              \
-	};                                                                                                                                 \
-	F##_ClassName##_Runner _ClassName##_RunnerInstance;                                                                                \
+
+#define _TEST_CLASS_IMPL_EXT(_ClassName, _TestDir, _BaseClass, _AsserterType, _TestFlags, _TestTags)                                              \
+	struct _ClassName;                                                                                                                            \
+	struct F##_ClassName##_Runner : public TTestRunner<_AsserterType>                                                                             \
+	{                                                                                                                                             \
+		F##_ClassName##_Runner()                                                                                                                  \
+			: TTestRunner(#_ClassName, __LINE__, __FILE__, _TestDir, _TestFlags, TTest<_ClassName, _AsserterType>::CreateTestClass, _TestTags) {  \
+				static_assert(  !!((_TestFlags) & EAutomationTestFlags_ApplicationContextMask),                                                   \
+								"CQTest has no application flag and will not run. See AutomationTest.h.");                                        \
+				static_assert(  !!(((_TestFlags) & EAutomationTestFlags_FilterMask) == EAutomationTestFlags::SmokeFilter) ||                      \
+								!!(((_TestFlags) & EAutomationTestFlags_FilterMask) == EAutomationTestFlags::EngineFilter) ||                     \
+								!!(((_TestFlags) & EAutomationTestFlags_FilterMask) == EAutomationTestFlags::ProductFilter) ||                    \
+								!!(((_TestFlags) & EAutomationTestFlags_FilterMask) == EAutomationTestFlags::PerfFilter) ||                       \
+								!!(((_TestFlags) & EAutomationTestFlags_FilterMask) == EAutomationTestFlags::StressFilter) ||                     \
+								!!(((_TestFlags) & EAutomationTestFlags_FilterMask) == EAutomationTestFlags::NegativeFilter),                     \
+								"All CQTests must have exactly 1 filter type specified. See AutomationTest.h.");                                  \
+		}                                                                                                                                         \
+	};                                                                                                                                            \
+	F##_ClassName##_Runner _ClassName##_RunnerInstance;                                                                                           \
 	struct _ClassName : public _BaseClass<_ClassName, _AsserterType>
 
-#define TEST_METHOD(_MethodName)                                                                       \
-	FFunctionRegistrar reg##_MethodName{ FString(#_MethodName), &DerivedType::_MethodName, __LINE__ }; \
+#define _TEST_CLASS_IMPL(_ClassName, _TestDir, _BaseClass, _AsserterType, _TestFlags)                  \
+		_TEST_CLASS_IMPL_EXT(_ClassName, _TestDir, _BaseClass, _AsserterType, _TestFlags, DefaultTags)
+
+#define TEST_METHOD_WITH_TAGS(_MethodName, _TestTags)                                                             \
+	FFunctionRegistrar reg##_MethodName{ FString(#_MethodName), &DerivedType::_MethodName, __LINE__, _TestTags }; \
 	void _MethodName()
+
+#define TEST_METHOD(_MethodName) TEST_METHOD_WITH_TAGS(_MethodName, DefaultTags)
+
 #else
-#define _TEST_CLASS_IMPL(_ClassName, _TestDir, _BaseClass, _AsserterType, _TestFlags) \
+#define _TEST_CLASS_IMPL_EXT(_ClassName, _TestDir, _BaseClass, _AsserterType, _TestFlags, _TestTags) \
 	struct _ClassName : public _BaseClass<_ClassName, _AsserterType>
 
-#define TEST_METHOD(_MethodName) void _MethodName()
+#define _TEST_CLASS_IMPL(_ClassName, _TestDir, _BaseClass, _AsserterType, _TestFlags) _TEST_CLASS_IMPL_EXT(_ClassName, _TestDir, _BaseClass, _AsserterType, _TestFlags, DefaultTags)
+
+#define TEST_METHOD_WITH_TAGS(_MethodName, _TestTags) void _MethodName()
+#define TEST_METHOD(_MethodName) TEST_METHOD_WITH_TAGS(_MethodName, DefaultTags)
 #endif
 
 #define TEST_CLASS_WITH_ASSERTS(_ClassName, _TestDir, _AsserterType) _TEST_CLASS_IMPL(_ClassName, _TestDir, TTest, _AsserterType, DefaultFlags)
-
+#define TEST_CLASS_WITH_ASSERTS_AND_TAGS(_ClassName, _TestDir, _AsserterType, _TestTags) _TEST_CLASS_IMPL_EXT(_ClassName, _TestDir, TTest, _AsserterType, DefaultFlags, _TestTags)
 #define TEST_CLASS(_ClassName, _TestDir) TEST_CLASS_WITH_ASSERTS(_ClassName, _TestDir, FNoDiscardAsserter)
+#define TEST_CLASS_WITH_TAGS(_ClassName, _TestDir, _TestTags) TEST_CLASS_WITH_ASSERTS_AND_TAGS(_ClassName, _TestDir, FNoDiscardAsserter, _TestTags)
 #define TEST_CLASS_WITH_BASE(_ClassName, _TestDir, _BaseClass) _TEST_CLASS_IMPL(_ClassName, _TestDir, _BaseClass, FNoDiscardAsserter, DefaultFlags)
+#define TEST_CLASS_WITH_BASE_AND_TAGS(_ClassName, _TestDir, _BaseClass, _TestTags) _TEST_CLASS_IMPL_EXT(_ClassName, _TestDir, _BaseClass, FNoDiscardAsserter, DefaultFlags, _TestTags)
 #define TEST_CLASS_WITH_FLAGS(_ClassName, _TestDir, _Flags) _TEST_CLASS_IMPL(_ClassName, _TestDir, TTest, FNoDiscardAsserter, _Flags)
+#define TEST_CLASS_WITH_FLAGS_AND_TAGS(_ClassName, _TestDir, _Flags, _TestTags) _TEST_CLASS_IMPL_EXT(_ClassName, _TestDir, TTest, FNoDiscardAsserter, _Flags, _TestTags)
 #define TEST_CLASS_WITH_BASE_AND_FLAGS(_ClassName, _TestDir, _BaseClass, _Flags) _TEST_CLASS_IMPL(_ClassName, _TestDir, _BaseClass, FNoDiscardAsserter, _Flags)
+#define TEST_CLASS_WITH_BASE_AND_FLAGS_AND_TAGS(_ClassName, _TestDir, _BaseClass, _Flags, _TestTags) _TEST_CLASS_IMPL_EXT(_ClassName, _TestDir, _BaseClass, FNoDiscardAsserter, _Flags, _TestTags)
 
 #define TEST(_TestName, _TestDir)        \
 	TEST_CLASS(_TestName, _TestDir)      \
@@ -273,6 +314,13 @@ struct TTest : TBaseTest<AsserterType>
 	};                                   \
 	void _TestName::_TestName##_Method()
 
+#define TEST_WITH_TAGS(_TestName, _TestDir, _TestTags)   \
+	TEST_CLASS_WITH_TAGS(_TestName, _TestDir, _TestTags) \
+	{                                                    \
+		TEST_METHOD(_TestName##_Method);                 \
+	};                                                   \
+	void _TestName::_TestName##_Method()
+
 #define ASSERT_THAT(_assertion) if (!Assert._assertion) {return;}
 #define ASSERT_FAIL(_Msg) Assert.Fail(_Msg); return;
 
diff --git a/Engine/Source/Developer/CQTest/Public/Impl/CQTest.inl b/Engine/Source/Developer/CQTest/Public/Impl/CQTest.inl
--- a/Engine/Source/Developer/CQTest/Public/Impl/CQTest.inl
+++ b/Engine/Source/Developer/CQTest/Public/Impl/CQTest.inl
@@ -127,15 +127,15 @@ namespace
 	};
 } // namespace
 
-
 template <typename AsserterType>
-inline TTestRunner<AsserterType>::TTestRunner(FString InName, int32 InLineNumber, const char* InFileName, FString InTestDir, EAutomationTestFlags InTestFlags, TTestInstanceGenerator<AsserterType> InFactory)
+inline TTestRunner<AsserterType>::TTestRunner(FString InName, int32 InLineNumber, const char* InFileName, FString InTestDir, EAutomationTestFlags InTestFlags, TTestInstanceGenerator<AsserterType> InFactory, FString InTestTags)
 	: FAutomationTestBase(InName, true)
 	, LineNumber(InLineNumber)
 	, FileName(FString(InFileName))
 	, TestDir(InTestDir)
 	, TestFlags(InTestFlags)
 	, TestInstanceFactory(InFactory)
+	, TestTags(InTestTags)
 {
 	bInitializing = true;
 	if (TestDir.Equals(GenerateTestDirectory))
diff --git a/Engine/Source/Developer/CQTest/README.md b/Engine/Source/Developer/CQTest/README.md
--- a/Engine/Source/Developer/CQTest/README.md
+++ b/Engine/Source/Developer/CQTest/README.md
@@ -105,9 +105,10 @@ CQTest comes with an additonal CQTestEnhancedInput plugin that provides addition
 
 To run tests provided by the plugins within the Unreal Editor
 - Launch the editor
+- Make sure the plugin "Code Quality Tests Unreal Test Plugin" is enabled
 - Find the Tools drop down and select Session Frontend
 - Navigate to the Automation tab
-- By default, the tests should be listed first under "Product.Plugins.CQTest"
+- By default, the tests should be listed first under "TestFramework.CQTest"
 - Select the tests you would like to run and press 'Start Tests'
 
 # Examples
@@ -180,12 +181,33 @@ For setup and teardown, or common state between multiple tests, or to group rela
 	};
 ```
 
-In addition to `TEST` and `TEST_CLASS` are 5 additional macros:
+Tests can be tagged using reserved tag macro definitions and they are specified in the format `[TagA][TagB][TagC]`.
+Test classes specify common tags that apply to all enclosed methods. In addition, each test method can specify their own tags.
+```cpp
+    #include "CQTest.h"
+	TEST_CLASS_WITH_TAGS(MyTaggedTest, "Game.MyGame.WithTags", "[FeatureA][SpecialTagA]") 
+	{
+		TEST_METHOD_WITH_TAGS(MyTestMethod, "[MethodSpecificTag]")
+		{
+			// This method can be filtered by tags FeatureA, SpecialTagA (common as defined by test class) and MethodSpecificTag (method specific tag).
+		}
+	}
+```
+
+In addition to `TEST` and `TEST_CLASS` are 13 additional macros:
+ - `TEST_METHOD_WITH_TAGS` - Extends `TEST_METHOD` with an additional argument for test tags.
+ - `TEST_WITH_TAGS` - Extends `TEST` with an additional argument for test tags.
+ - `TEST_CLASS_WITH_TAGS` - Extends `TEST_CLASS` with an additional argument for test tags.
  - `TEST_CLASS_WITH_ASSERTS` - Macro which allows this test object to use a custom asserter.  More information about how to use this macro can be found [in the section below regarding assertions](#assertions)
+ - `TEST_CLASS_WITH_ASSERTS_AND_TAGS` - Extends `TEST_CLASS_WITH_ASSERTS` with an additional argument for test tags.
  - `TEST_CLASS_WITH_BASE` - Macro which allows this test object to inherit from a different test object.  More information about how to use this macro can be found [in the section below regarding custom test classes](#base-test-class)
+ - `TEST_CLASS_WITH_BASE_AND_TAGS` - Extends `TEST_CLASS_WITH_BASE` with an additional argument for test tags.
  - `TEST_CLASS_WITH_FLAGS` - Macro which allows the use of different automation test flags to be specified.  Useful for when tests can only run under a certain context or grouped under a specific filter.  The default flags for `TEST` and `TEST_CLASS` are `EAutomationTestFlags_ApplicationContextMask | EAutomationTestFlags::ProductFilter`.  Additional information regarding the available flags can be found [in the online documentation](https://dev.epicgames.com/documentation/en-us/unreal-engine/API/Runtime/Core/Misc/EAutomationTestFlags/Type)
+ - `TEST_CLASS_WITH_FLAGS_AND_TAGS` - Extends - `TEST_CLASS_WITH_FLAGS` with an additional argument for test tags.
  - `TEST_CLASS_WITH_BASE_AND_FLAGS` - Macro which allows for this test object to inherit from a different test object and allows for custom automation test flags to be specified.
- - `_TEST_CLASS_IMPL` - Base macro which is used by the above macros to specify a custom asserter, a test object to inherit from a different test object, and allows for custom automation test flags to be specified.
+ - `TEST_CLASS_WITH_BASE_AND_FLAGS_AND_TAGS` - Extends - `TEST_CLASS_WITH_BASE_AND_FLAGS` with an additional argument for test tags.
+ - `_TEST_CLASS_IMPL_EXT` - Base macro which is used by the above macros to specify a custom asserter, a test object to inherit from a different test object, and allows for custom automation test flags to be specified as well as test tags.
+ - `_TEST_CLASS_IMPL` - Expands `_TEST_CLASS_IMPL_EXT` with default tags.
 
 Test Directory determines where in the Automation tab the tests appear.  In the example above, we specify "Game.MyGame", but you may also have an auto-generated test directory based on the folder structure.
 ```cpp
