Add support for pointers and TOptional values to CQTestConvert::ToString

#rnx
PR #13791

diff --git a/Engine/Plugins/Tests/CQTest/Source/CQTestTests/Private/Assert/CQTestConvertTests.cpp b/Engine/Plugins/Tests/CQTest/Source/CQTestTests/Private/Assert/CQTestConvertTests.cpp
--- a/Engine/Plugins/Tests/CQTest/Source/CQTestTests/Private/Assert/CQTestConvertTests.cpp
+++ b/Engine/Plugins/Tests/CQTest/Source/CQTestTests/Private/Assert/CQTestConvertTests.cpp
@@ -282,4 +282,23 @@ TEST_CLASS(TestConvertToString, "TestFramework.CQTest.Core")
 		auto asStr = CQTestConvert::ToString(value);
 		ASSERT_THAT(IsTrue(asStr.Contains(TEXT("3.14"))));
 	}
+
+	TEST_METHOD(Pointer_ConvertToString_ReturnsPointerAsString)
+	{
+		int32 Value = 3;
+		int32* PtrToValue = &Value;
+		ASSERT_THAT(AreEqual(FString::Printf(TEXT("0x%p"), PtrToValue), CQTestConvert::ToString(PtrToValue)));
+	}
+
+	TEST_METHOD(OptionalWithValue_ConvertToString_ReturnsValueAsString)
+	{
+		TOptional<int32> Value(3);
+		ASSERT_THAT(AreEqual(FString(TEXT("3")), CQTestConvert::ToString(Value)));
+	}
+
+	TEST_METHOD(OptionalWithoutValue_ConvertToString_ReturnsNotSetAsString)
+	{
+		TOptional<int32> Value;
+		ASSERT_THAT(AreEqual(FString(TEXT("[NotSet]")), CQTestConvert::ToString(Value)));
+	}
 };
diff --git a/Engine/Source/Developer/CQTest/Public/Assert/CQTestConvert.h b/Engine/Source/Developer/CQTest/Public/Assert/CQTestConvert.h
--- a/Engine/Source/Developer/CQTest/Public/Assert/CQTestConvert.h
+++ b/Engine/Source/Developer/CQTest/Public/Assert/CQTestConvert.h
@@ -2,6 +2,7 @@
 #pragma once
 
 #include "Containers/UnrealString.h"
+#include "Misc/Optional.h"
 #include "Misc/StringBuilder.h"
 #include "UObject/NameTypes.h"
 #include <sstream>
@@ -99,6 +100,11 @@ namespace CQTestConvert
 			// Fallback for enum type
 			return ToString(static_cast<std::underlying_type_t<T>>(Input));
 		}
+		else if constexpr (std::is_pointer_v<T>)
+		{
+			// Fallback for pointer type
+			return FString::Printf(TEXT("0x%p"), (void*)Input);
+		}
 		else
 		{
 			ensureMsgf(false, TEXT("Did not find ToString, ostream operator, or CQTestConvert::ToString template specialization for provided type. Cast provided type to something else, or see CqTestConvertTests.cpp for examples"));
@@ -137,4 +143,17 @@ namespace CQTestConvert
 			return TEXT("{") + CQTestConvert::ToString(kvp.Key) + TEXT(":") + CQTestConvert::ToString(kvp.Value) + TEXT("}");
 		});
 	}
+
+	template<typename TInnerType>
+	inline FString ToString(const TOptional<TInnerType>& Input)
+	{
+		if (Input.IsSet())
+		{
+			return ToString(Input.GetValue());
+		}
+		else
+		{
+			return TEXT("[NotSet]");
+		}
+	}
 }
