[CQTest] Added TEST_CLASS_WITH_BASE_AND_ASSERTS, renamed templated parameters to allow for aliasing

#Synced-CL 48021612

diff --git a/Engine/Source/Developer/CQTest/Public/CQTest.h b/Engine/Source/Developer/CQTest/Public/CQTest.h
--- a/Engine/Source/Developer/CQTest/Public/CQTest.h
+++ b/Engine/Source/Developer/CQTest/Public/CQTest.h
@@ -76,9 +76,11 @@ concept HasAfterAll = requires(T t) { { T::AfterAll(FString()) }; };
 
 using BeforeAfterAllFunc = void(*)(const FString&);
 
-template <typename AsserterType>
+template <typename InAsserterType>
 struct TBaseTest
 {
+	using AsserterType = InAsserterType;
+
 	TBaseTest(FAutomationTestBase& TestRunner, bool bInitializing);
 	TBaseTest(const TBaseTest<AsserterType>& other) = delete;
 	TBaseTest& operator=(const TBaseTest<AsserterType>& rhs) = delete;
@@ -135,15 +137,17 @@ struct TBaseTest
 	FTestCommandBuilder TestCommandBuilder;
 };
 
-template <typename AsserterType>
+template <typename InAsserterType>
 struct TTestRunner;
 
-template <typename AsserterType>
-using TTestInstanceGenerator = TUniquePtr<TBaseTest<AsserterType>>(*)(TTestRunner<AsserterType>&);
+template <typename InAsserterType>
+using TTestInstanceGenerator = TUniquePtr<TBaseTest<InAsserterType>>(*)(TTestRunner<InAsserterType>&);
 
-template <typename AsserterType>
+template <typename InAsserterType>
 struct TTestRunner : public FAutomationTestBase
 {
+	using AsserterType = InAsserterType;
+
 	TTestRunner(FString Name, int32 LineNumber, const char* FileName, FString TestDir, EAutomationTestFlags TestFlags, TTestInstanceGenerator<AsserterType> Factory, FString TestTags);
 	~TTestRunner();
 
@@ -302,6 +306,7 @@ struct TTest : TBaseTest<AsserterType>
 #define TEST_CLASS_WITH_TAGS(_ClassName, _TestDir, _TestTags) TEST_CLASS_WITH_ASSERTS_AND_TAGS(_ClassName, _TestDir, FNoDiscardAsserter, _TestTags)
 #define TEST_CLASS_WITH_BASE(_ClassName, _TestDir, _BaseClass) _TEST_CLASS_IMPL(_ClassName, _TestDir, _BaseClass, FNoDiscardAsserter, DefaultFlags)
 #define TEST_CLASS_WITH_BASE_AND_TAGS(_ClassName, _TestDir, _BaseClass, _TestTags) _TEST_CLASS_IMPL_EXT(_ClassName, _TestDir, _BaseClass, FNoDiscardAsserter, DefaultFlags, _TestTags)
+#define TEST_CLASS_WITH_BASE_AND_ASSERTS(_ClassName, _TestDir, _BaseClass, _AsserterType) _TEST_CLASS_IMPL(_ClassName, _TestDir, _BaseClass, _AsserterType, DefaultFlags)
 #define TEST_CLASS_WITH_FLAGS(_ClassName, _TestDir, _Flags) _TEST_CLASS_IMPL(_ClassName, _TestDir, TTest, FNoDiscardAsserter, _Flags)
 #define TEST_CLASS_WITH_FLAGS_AND_TAGS(_ClassName, _TestDir, _Flags, _TestTags) _TEST_CLASS_IMPL_EXT(_ClassName, _TestDir, TTest, FNoDiscardAsserter, _Flags, _TestTags)
 #define TEST_CLASS_WITH_BASE_AND_FLAGS(_ClassName, _TestDir, _BaseClass, _Flags) _TEST_CLASS_IMPL(_ClassName, _TestDir, _BaseClass, FNoDiscardAsserter, _Flags)
